UE4: Non-unity fix

diff --git a/Engine/Source/Editor/BehaviorTreeEditor/Private/DetailCustomizations/BlackboardSelectorDetails.cpp b/Engine/Source/Editor/BehaviorTreeEditor/Private/DetailCustomizations/BlackboardSelectorDetails.cpp
--- a/Engine/Source/Editor/BehaviorTreeEditor/Private/DetailCustomizations/BlackboardSelectorDetails.cpp
+++ b/Engine/Source/Editor/BehaviorTreeEditor/Private/DetailCustomizations/BlackboardSelectorDetails.cpp
@@ -2,6 +2,7 @@
 
 #include "../BehaviorTreeEditorPrivatePCH.h"
 #include "BlackboardSelectorDetails.h"
+#include "IPropertyUtilities.h"
 
 #define LOCTEXT_NAMESPACE "BlackboardSelectorDetails"
 
