Missing file from previous checkin

diff --git a/Engine/Source/Editor/BehaviorTreeEditor/Classes/BehaviorTreeDecoratorGraphNode.h b/Engine/Source/Editor/BehaviorTreeEditor/Classes/BehaviorTreeDecoratorGraphNode.h
--- a/Engine/Source/Editor/BehaviorTreeEditor/Classes/BehaviorTreeDecoratorGraphNode.h
+++ b/Engine/Source/Editor/BehaviorTreeEditor/Classes/BehaviorTreeDecoratorGraphNode.h
@@ -3,6 +3,7 @@
 #pragma once
 
 #include "BehaviorTree/BTCompositeNode.h"
+#include "EdGraph/EdGraphNode.h"
 #include "BehaviorTreeDecoratorGraphNode.generated.h"
 
 
