improved lazy loading of blueprint based behavior tree nodes
TTP# 337875

diff --git a/Engine/Source/Editor/BehaviorTreeEditor/Classes/BehaviorTreeDecoratorGraphNode_Decorator.h b/Engine/Source/Editor/BehaviorTreeEditor/Classes/BehaviorTreeDecoratorGraphNode_Decorator.h
--- a/Engine/Source/Editor/BehaviorTreeEditor/Classes/BehaviorTreeDecoratorGraphNode_Decorator.h
+++ b/Engine/Source/Editor/BehaviorTreeEditor/Classes/BehaviorTreeDecoratorGraphNode_Decorator.h
@@ -23,6 +23,7 @@ class UBehaviorTreeDecoratorGraphNode_Decorator : public UBehaviorTreeDecoratorG
 	virtual void PostEditImport() OVERRIDE;
 	virtual void PrepareForCopying() OVERRIDE;
 	void PostCopyNode();
+	bool RefreshNodeClass();
 
 private:
 
diff --git a/Engine/Source/Editor/BehaviorTreeEditor/Classes/BehaviorTreeEditorTypes.h b/Engine/Source/Editor/BehaviorTreeEditor/Classes/BehaviorTreeEditorTypes.h
--- a/Engine/Source/Editor/BehaviorTreeEditor/Classes/BehaviorTreeEditorTypes.h
+++ b/Engine/Source/Editor/BehaviorTreeEditor/Classes/BehaviorTreeEditorTypes.h
@@ -13,8 +13,11 @@ struct FAbortDrawHelper
 	FAbortDrawHelper() : AbortStart(MAX_uint16), AbortEnd(MAX_uint16), SearchStart(MAX_uint16), SearchEnd(MAX_uint16) {}
 };
 
+USTRUCT()
 struct FClassData
 {
+	GENERATED_USTRUCT_BODY()
+
 	FClassData() {}
 	FClassData(UClass* InClass, const FString& InDeprecatedMessage) : 
 		bIsHidden(0), bHideParent(0), Class(InClass), DeprecatedMessage(InDeprecatedMessage)
@@ -25,11 +28,13 @@ struct FClassData
 
 	FString ToString() const;
 	FString GetClassName() const;
-	UClass* GetClass();
+	UClass* GetClass(bool bSilent = false);
 	bool IsAbstract() const;
-
+	
+	FORCEINLINE bool IsBlueprint() const { return AssetName.Len() > 0; }
 	FORCEINLINE bool IsDeprecated() const { return DeprecatedMessage.Len() > 0; }
 	FORCEINLINE FString GetDeprecatedMessage() const { return DeprecatedMessage; }
+	FORCEINLINE FString GetPackageName() const { return GeneratedClassPackage; }
 
 	/** set when child class masked this one out (e.g. always use game specific class instead of engine one) */
 	uint32 bIsHidden : 1;
@@ -43,10 +48,13 @@ private:
 	TWeakObjectPtr<UClass> Class;
 
 	/** path to class if it's not loaded yet */
+	UPROPERTY()
 	FString AssetName;
+	UPROPERTY()
 	FString GeneratedClassPackage;
 
 	/** resolved name of class from asset data */
+	UPROPERTY()
 	FString ClassName;
 
 	/** message for deprecated class */
@@ -60,10 +68,14 @@ struct FClassDataNode
 
 	TSharedPtr<FClassDataNode> ParentNode;
 	TArray<TSharedPtr<FClassDataNode> > SubNodes;
+
+	void AddUniqueSubNode(TSharedPtr<FClassDataNode> SubNode);
 };
 
 struct FClassBrowseHelper
 {
+	DECLARE_MULTICAST_DELEGATE(FOnPackageListUpdated);
+
 	FClassBrowseHelper();
 	~FClassBrowseHelper();
 
@@ -73,10 +85,15 @@ struct FClassBrowseHelper
 	void OnAssetAdded(const class FAssetData& AssetData);
 	void OnAssetRemoved(const class FAssetData& AssetData);
 	void InvalidateCache();
+	
+	static void AddUnknownClass(const FClassData& ClassData);
+	static bool IsClassKnown(const FClassData& ClassData);
+	static FOnPackageListUpdated OnPackageListUpdated;
 
 private:
 
 	TSharedPtr<FClassDataNode> RootNode;
+	static TArray<FName> UnknownPackages;
 	
 	TSharedPtr<FClassDataNode> CreateClassDataNode(const class FAssetData& AssetData);
 	TSharedPtr<FClassDataNode> FindBaseClassNode(TSharedPtr<FClassDataNode> Node, const FString& ClassName);
@@ -87,6 +104,7 @@ private:
 	void AddClassGraphChildren(TSharedPtr<FClassDataNode> Node, TArray<TSharedPtr<FClassDataNode> >& NodeList);
 
 	bool IsHidingParentClass(UClass* Class);
+	bool IsPackageSaved(FName PackageName);
 };
 
 struct FCompareNodeXLocation
diff --git a/Engine/Source/Editor/BehaviorTreeEditor/Classes/BehaviorTreeGraph.h b/Engine/Source/Editor/BehaviorTreeEditor/Classes/BehaviorTreeGraph.h
--- a/Engine/Source/Editor/BehaviorTreeEditor/Classes/BehaviorTreeGraph.h
+++ b/Engine/Source/Editor/BehaviorTreeEditor/Classes/BehaviorTreeGraph.h
@@ -26,6 +26,7 @@ class UBehaviorTreeGraph : public UEdGraph
 	void UpdatePinConnectionTypes();
 	void UpdateDeprecatedNodes();
 	bool UpdateInjectedNodes();
+	bool UpdateUnknownNodeClasses();
 	class UEdGraphNode* FindInjectedNode(int32 Index);
 	void ReplaceNodeConnections(UEdGraphNode* OldNode, UEdGraphNode* NewNode);
 	void RebuildExecutionOrder();
diff --git a/Engine/Source/Editor/BehaviorTreeEditor/Classes/BehaviorTreeGraphNode.h b/Engine/Source/Editor/BehaviorTreeEditor/Classes/BehaviorTreeGraphNode.h
--- a/Engine/Source/Editor/BehaviorTreeEditor/Classes/BehaviorTreeGraphNode.h
+++ b/Engine/Source/Editor/BehaviorTreeEditor/Classes/BehaviorTreeGraphNode.h
@@ -74,10 +74,17 @@ class UBehaviorTreeGraphNode : public UEdGraphNode
 	/** add subnode */
 	void AddSubNode(UBehaviorTreeGraphNode* NodeTemplate, class UEdGraph* ParentGraph);
 
+	/** reinitialize node instance */
+	virtual bool RefreshNodeClass();
+
 	/** gets icon resource name for title bar */
 	virtual FName GetNameIcon() const;
 
+	/** check if node has any errors, used for assigning colors on graph */
+	virtual bool HasErrors() const;
+
 	/** instance class */
+	UPROPERTY()
 	struct FClassData ClassData;
 
 	/** if set, this node was injected from subtree and shouldn't be edited */
diff --git a/Engine/Source/Editor/BehaviorTreeEditor/Classes/BehaviorTreeGraphNode_CompositeDecorator.h b/Engine/Source/Editor/BehaviorTreeEditor/Classes/BehaviorTreeGraphNode_CompositeDecorator.h
--- a/Engine/Source/Editor/BehaviorTreeEditor/Classes/BehaviorTreeGraphNode_CompositeDecorator.h
+++ b/Engine/Source/Editor/BehaviorTreeEditor/Classes/BehaviorTreeGraphNode_CompositeDecorator.h
@@ -33,6 +33,8 @@ class UBehaviorTreeGraphNode_CompositeDecorator : public UBehaviorTreeGraphNode
 	virtual void PostLoad() OVERRIDE;
 	virtual UEdGraph* GetBoundGraph() const { return BoundGraph; }
 	virtual bool IsSubNode() const OVERRIDE;
+	virtual bool HasErrors() const OVERRIDE;
+	virtual bool RefreshNodeClass() OVERRIDE;
 
 	virtual void PrepareForCopying() OVERRIDE;
 	virtual void PostCopyNode() OVERRIDE;
diff --git a/Engine/Source/Editor/BehaviorTreeEditor/Classes/BehaviorTreeGraphNode_Root.h b/Engine/Source/Editor/BehaviorTreeEditor/Classes/BehaviorTreeGraphNode_Root.h
--- a/Engine/Source/Editor/BehaviorTreeEditor/Classes/BehaviorTreeGraphNode_Root.h
+++ b/Engine/Source/Editor/BehaviorTreeEditor/Classes/BehaviorTreeGraphNode_Root.h
@@ -14,7 +14,9 @@ class UBehaviorTreeGraphNode_Root : public UBehaviorTreeGraphNode
 	virtual void PostPlacedNewNode() OVERRIDE;
 	virtual void AllocateDefaultPins() OVERRIDE;
 	virtual bool CanDuplicateNode() const OVERRIDE { return false; }
-	virtual bool CanUserDeleteNode() const OVERRIDE { return false; }
+	virtual bool CanUserDeleteNode() const OVERRIDE{ return false; }
+	virtual bool HasErrors() const OVERRIDE { return false; }
+	virtual bool RefreshNodeClass() OVERRIDE { return false; }
 	virtual FText GetNodeTitle(ENodeTitleType::Type TitleType) const OVERRIDE;
 	/** gets icon resource name for title bar */
 	virtual FName GetNameIcon() const OVERRIDE;
diff --git a/Engine/Source/Editor/BehaviorTreeEditor/Private/BehaviorTreeDecoratorGraphNode_Decorator.cpp b/Engine/Source/Editor/BehaviorTreeEditor/Private/BehaviorTreeDecoratorGraphNode_Decorator.cpp
--- a/Engine/Source/Editor/BehaviorTreeEditor/Private/BehaviorTreeDecoratorGraphNode_Decorator.cpp
+++ b/Engine/Source/Editor/BehaviorTreeEditor/Private/BehaviorTreeDecoratorGraphNode_Decorator.cpp
@@ -89,3 +89,22 @@ void UBehaviorTreeDecoratorGraphNode_Decorator::ResetNodeOwner()
 		NodeInstance->Rename(NULL, BT, REN_DontCreateRedirectors | REN_DoNotDirty);
 	}
 }
+
+bool UBehaviorTreeDecoratorGraphNode_Decorator::RefreshNodeClass()
+{
+	bool bUpdated = false;
+	if (NodeInstance == NULL)
+	{
+		if (FClassBrowseHelper::IsClassKnown(ClassData))
+		{
+			PostPlacedNewNode();
+			bUpdated = (NodeInstance != NULL);
+		}
+		else
+		{
+			FClassBrowseHelper::AddUnknownClass(ClassData);
+		}
+	}
+
+	return bUpdated;
+}
diff --git a/Engine/Source/Editor/BehaviorTreeEditor/Private/BehaviorTreeEditor.cpp b/Engine/Source/Editor/BehaviorTreeEditor/Private/BehaviorTreeEditor.cpp
--- a/Engine/Source/Editor/BehaviorTreeEditor/Private/BehaviorTreeEditor.cpp
+++ b/Engine/Source/Editor/BehaviorTreeEditor/Private/BehaviorTreeEditor.cpp
@@ -138,6 +138,7 @@ FBehaviorTreeEditor::FBehaviorTreeEditor()
 
 	// listen for package change events to update injected nodes
 	UPackage::PackageSavedEvent.AddRaw(this, &FBehaviorTreeEditor::OnPackageSaved);
+	FClassBrowseHelper::OnPackageListUpdated.AddRaw(this, &FBehaviorTreeEditor::OnClassListUpdated);
 
 	bShowDecoratorRangeLower = false;
 	bShowDecoratorRangeSelf = false;
@@ -154,6 +155,7 @@ FBehaviorTreeEditor::~FBehaviorTreeEditor()
 	}
 
 	UPackage::PackageSavedEvent.RemoveRaw(this, &FBehaviorTreeEditor::OnPackageSaved);
+	FClassBrowseHelper::OnPackageListUpdated.RemoveRaw(this, &FBehaviorTreeEditor::OnClassListUpdated);
 	Debugger.Reset();
 }
 
@@ -293,6 +295,7 @@ void FBehaviorTreeEditor::InitBehaviorTreeEditor( const EToolkitMode::Type Mode,
 
 		MyGraph->UpdateBlackboardChange();
 		MyGraph->UpdateInjectedNodes();
+		MyGraph->UpdateUnknownNodeClasses();
 
 		TSharedRef<FTabPayload_UObject> Payload = FTabPayload_UObject::Make(MyGraph);
 		DocumentManager.OpenDocument(Payload, FDocumentTracker::OpenNewDocument);
@@ -1093,6 +1096,22 @@ void FBehaviorTreeEditor::OnPackageSaved(const FString& PackageFileName, UObject
 	}
 }
 
+void FBehaviorTreeEditor::OnClassListUpdated()
+{
+	UBehaviorTreeGraph* MyGraph = Script ? Cast<UBehaviorTreeGraph>(Script->BTGraph) : NULL;
+	if (MyGraph)
+	{
+		const bool bUpdated = MyGraph->UpdateUnknownNodeClasses();
+		if (bUpdated)
+		{
+			FGraphPanelSelectionSet CurrentSelection = GetSelectedNodes();
+			OnSelectedNodesChanged(CurrentSelection);
+
+			MyGraph->UpdateAsset(UBehaviorTreeGraph::ClearDebuggerFlags);
+		}
+	}
+}
+
 FGraphPanelSelectionSet FBehaviorTreeEditor::GetSelectedNodes() const
 {
 	FGraphPanelSelectionSet CurrentSelection;
diff --git a/Engine/Source/Editor/BehaviorTreeEditor/Private/BehaviorTreeEditor.h b/Engine/Source/Editor/BehaviorTreeEditor/Private/BehaviorTreeEditor.h
--- a/Engine/Source/Editor/BehaviorTreeEditor/Private/BehaviorTreeEditor.h
+++ b/Engine/Source/Editor/BehaviorTreeEditor/Private/BehaviorTreeEditor.h
@@ -81,6 +81,7 @@ public:
 	bool IsPropertyEditable() const;
 	void OnPackageSaved(const FString& PackageFileName, UObject* Outer);
 	void OnFinishedChangingProperties(const FPropertyChangedEvent& PropertyChangedEvent);
+	void OnClassListUpdated();
 
 	void UpdateToolbar();
 	bool IsDebuggerReady() const;
diff --git a/Engine/Source/Editor/BehaviorTreeEditor/Private/BehaviorTreeEditorTypes.cpp b/Engine/Source/Editor/BehaviorTreeEditor/Private/BehaviorTreeEditorTypes.cpp
--- a/Engine/Source/Editor/BehaviorTreeEditor/Private/BehaviorTreeEditorTypes.cpp
+++ b/Engine/Source/Editor/BehaviorTreeEditor/Private/BehaviorTreeEditorTypes.cpp
@@ -47,7 +47,7 @@ bool FClassData::IsAbstract() const
 	return Class.IsValid() ? Class.Get()->HasAnyClassFlags(CLASS_Abstract) : false;
 }
 
-UClass* FClassData::GetClass()
+UClass* FClassData::GetClass(bool bSilent)
 {
 	UClass* RetClass = Class.Get();
 	if (RetClass == NULL && GeneratedClassPackage.Len())
@@ -64,17 +64,23 @@ UClass* FClassData::GetClass()
 			GWarn->EndSlowTask();
 
 			UBlueprint* BlueprintOb = Cast<UBlueprint>(Object);
-			RetClass = BlueprintOb ? *BlueprintOb->GeneratedClass : Object->GetClass();
+			RetClass = BlueprintOb ? *BlueprintOb->GeneratedClass :
+				Object ? Object->GetClass() : 
+				NULL;
+
 			Class = RetClass;
 		}
 		else
 		{
 			GWarn->EndSlowTask();
 
-			FMessageLog EditorErrors("EditorErrors");
-			EditorErrors.Error(LOCTEXT("PackageLoadFail", "Package Load Failed"));
-			EditorErrors.Info(FText::FromString(GeneratedClassPackage));
-			EditorErrors.Notify(LOCTEXT("PackageLoadFail", "Package Load Failed"));
+			if (!bSilent)
+			{
+				FMessageLog EditorErrors("EditorErrors");
+				EditorErrors.Error(LOCTEXT("PackageLoadFail", "Package Load Failed"));
+				EditorErrors.Info(FText::FromString(GeneratedClassPackage));
+				EditorErrors.Notify(LOCTEXT("PackageLoadFail", "Package Load Failed"));
+			}
 		}
 	}
 
@@ -82,6 +88,21 @@ UClass* FClassData::GetClass()
 }
 
 //////////////////////////////////////////////////////////////////////////
+FClassBrowseHelper::FOnPackageListUpdated FClassBrowseHelper::OnPackageListUpdated;
+TArray<FName> FClassBrowseHelper::UnknownPackages;
+
+void FClassDataNode::AddUniqueSubNode(TSharedPtr<FClassDataNode> SubNode)
+{
+	for (int32 Idx = 0; Idx < SubNodes.Num(); Idx++)
+	{
+		if (SubNode->Data.GetClassName() == SubNodes[Idx]->Data.GetClassName())
+		{
+			return;
+		}
+	}
+
+	SubNodes.Add(SubNode);
+}
 
 FClassBrowseHelper::FClassBrowseHelper()
 {
@@ -148,12 +169,33 @@ FString FClassBrowseHelper::GetDeprecationMessage(const UClass* Class)
 	return DeprecatedMessage;
 }
 
+DEFINE_LOG_CATEGORY_STATIC(LogBTDebug3, Log, All);
+
+bool FClassBrowseHelper::IsClassKnown(const FClassData& ClassData)
+{
+	return !ClassData.IsBlueprint() || !UnknownPackages.Contains(*ClassData.GetPackageName());
+}
+
+void FClassBrowseHelper::AddUnknownClass(const FClassData& ClassData)
+{
+	if (ClassData.IsBlueprint())
+	{
+		UnknownPackages.AddUnique(*ClassData.GetPackageName());
+	}
+}
+
 bool FClassBrowseHelper::IsHidingParentClass(UClass* Class)
 {
 	static FName MetaHideParent = TEXT("HideParentNode");
 	return Class && Class->HasAnyClassFlags(CLASS_Native) && Class->HasMetaData(MetaHideParent);
 }
 
+bool FClassBrowseHelper::IsPackageSaved(FName PackageName)
+{
+	UPackage* Package = LoadPackage(NULL, *PackageName.ToString(), LOAD_None);
+	return Package != NULL;
+}
+
 void FClassBrowseHelper::OnAssetAdded(const class FAssetData& AssetData)
 {
 	TSharedPtr<FClassDataNode> Node = CreateClassDataNode(AssetData);
@@ -162,11 +204,26 @@ void FClassBrowseHelper::OnAssetAdded(const class FAssetData& AssetData)
 	if (Node.IsValid())
 	{
 		ParentNode = FindBaseClassNode(RootNode, Node->ParentClassName);
+
+		if (!IsPackageSaved(AssetData.PackageName))
+		{
+			UnknownPackages.AddUnique(AssetData.PackageName);
+		}
+		else
+		{
+			const int32 PrevListCount = UnknownPackages.Num();
+			UnknownPackages.RemoveSingleSwap(AssetData.PackageName);
+
+			if (UnknownPackages.Num() != PrevListCount)
+			{
+				FClassBrowseHelper::OnPackageListUpdated.Broadcast();
+			}
+		}
 	}
 
 	if (ParentNode.IsValid())
 	{
-		ParentNode->SubNodes.Add(Node);
+		ParentNode->AddUniqueSubNode(Node);
 		Node->ParentNode = ParentNode;
 	}
 }
diff --git a/Engine/Source/Editor/BehaviorTreeEditor/Private/BehaviorTreeGraph.cpp b/Engine/Source/Editor/BehaviorTreeEditor/Private/BehaviorTreeGraph.cpp
--- a/Engine/Source/Editor/BehaviorTreeEditor/Private/BehaviorTreeGraph.cpp
+++ b/Engine/Source/Editor/BehaviorTreeEditor/Private/BehaviorTreeGraph.cpp
@@ -750,6 +750,37 @@ void UBehaviorTreeGraph::UpdateAbortHighlight(struct FAbortDrawHelper& Mode0, st
 	}
 }
 
+bool UBehaviorTreeGraph::UpdateUnknownNodeClasses()
+{
+	bool bUpdated = false;
+	for (int32 NodeIdx = 0; NodeIdx < Nodes.Num(); NodeIdx++)
+	{
+		UBehaviorTreeGraphNode* MyNode = Cast<UBehaviorTreeGraphNode>(Nodes[NodeIdx]);
+		const bool bUpdatedNode = MyNode->RefreshNodeClass();
+		bUpdated = bUpdated || bUpdatedNode;
+
+		for (int32 SubNodeIdx = 0; SubNodeIdx < MyNode->Decorators.Num(); SubNodeIdx++)
+		{
+			if (MyNode->Decorators[SubNodeIdx])
+			{
+				const bool bUpdatedSubNode = MyNode->Decorators[SubNodeIdx]->RefreshNodeClass();
+				bUpdated = bUpdated || bUpdatedSubNode;
+			}
+		}
+
+		for (int32 SubNodeIdx = 0; SubNodeIdx < MyNode->Services.Num(); SubNodeIdx++)
+		{
+			if (MyNode->Services[SubNodeIdx])
+			{
+				const bool bUpdatedSubNode = MyNode->Services[SubNodeIdx]->RefreshNodeClass();
+				bUpdated = bUpdated || bUpdatedSubNode;
+			}
+		}
+	}
+
+	return bUpdated;
+}
+
 bool UBehaviorTreeGraph::UpdateInjectedNodes()
 {
 	bool bHasUpdated = false;
diff --git a/Engine/Source/Editor/BehaviorTreeEditor/Private/BehaviorTreeGraphNode.cpp b/Engine/Source/Editor/BehaviorTreeEditor/Private/BehaviorTreeGraphNode.cpp
--- a/Engine/Source/Editor/BehaviorTreeEditor/Private/BehaviorTreeGraphNode.cpp
+++ b/Engine/Source/Editor/BehaviorTreeEditor/Private/BehaviorTreeGraphNode.cpp
@@ -42,7 +42,7 @@ void UBehaviorTreeGraphNode::AllocateDefaultPins()
 
 void UBehaviorTreeGraphNode::PostPlacedNewNode()
 {
-	UClass* NodeClass = ClassData.GetClass();
+	UClass* NodeClass = ClassData.GetClass(true);
 	if (NodeClass)
 	{
 		UBehaviorTree* BT = Cast<UBehaviorTree>(GetBehaviorTreeGraph()->GetOuter());
@@ -108,12 +108,17 @@ FString	UBehaviorTreeGraphNode::GetDescription() const
 		return Node->GetStaticDescription();
 	}
 
-	return FString();
+	FString StoredClassName = ClassData.GetClassName();
+	StoredClassName.RemoveFromEnd(TEXT("_C"));
+
+	return FString::Printf(TEXT("%s: %s"), *StoredClassName,
+		*LOCTEXT("NodeClassError", "Class not found, make sure it's saved!").ToString());
 }
 
 FString UBehaviorTreeGraphNode::GetTooltip() const
 {
-	FString TooltipDesc = bHasObserverError ? LOCTEXT("ObserverError", "Observer has invalid abort setting!").ToString() :		
+	FString TooltipDesc = !NodeInstance ? LOCTEXT("NodeClassError", "Class not found, make sure it's saved!").ToString() :
+		bHasObserverError ? LOCTEXT("ObserverError", "Observer has invalid abort setting!").ToString() :		
 		(DebuggerRuntimeDescription.Len() > 0) ? DebuggerRuntimeDescription : ErrorMessage;
 
 	if (bInjectedNode)
@@ -391,4 +396,29 @@ FName UBehaviorTreeGraphNode::GetNameIcon() const
 	return FName("BTEditor.Graph.BTNode.Icon");
 }
 
+bool UBehaviorTreeGraphNode::RefreshNodeClass()
+{
+	bool bUpdated = false;
+	if (NodeInstance == NULL)
+	{
+		if (FClassBrowseHelper::IsClassKnown(ClassData))
+		{
+			PostPlacedNewNode();
+			bUpdated = (NodeInstance != NULL);
+		}
+		else
+		{
+			FClassBrowseHelper::AddUnknownClass(ClassData);
+		}
+	}
+
+	return bUpdated;
+}
+
+bool UBehaviorTreeGraphNode::HasErrors() const
+{
+	return bHasObserverError || ErrorMessage.Len() > 0 || NodeInstance == NULL;
+}
+
+
 #undef LOCTEXT_NAMESPACE
diff --git a/Engine/Source/Editor/BehaviorTreeEditor/Private/BehaviorTreeGraphNode_CompositeDecorator.cpp b/Engine/Source/Editor/BehaviorTreeEditor/Private/BehaviorTreeGraphNode_CompositeDecorator.cpp
--- a/Engine/Source/Editor/BehaviorTreeEditor/Private/BehaviorTreeGraphNode_CompositeDecorator.cpp
+++ b/Engine/Source/Editor/BehaviorTreeEditor/Private/BehaviorTreeGraphNode_CompositeDecorator.cpp
@@ -93,6 +93,31 @@ void UBehaviorTreeGraphNode_CompositeDecorator::PostCopyNode()
 	}
 }
 
+bool UBehaviorTreeGraphNode_CompositeDecorator::RefreshNodeClass()
+{
+	bool bUpdated = false;
+
+	if (BoundGraph)
+	{
+		for (int32 i = 0; i < BoundGraph->Nodes.Num(); i++)
+		{
+			UBehaviorTreeDecoratorGraphNode_Decorator* Node = Cast<UBehaviorTreeDecoratorGraphNode_Decorator>(BoundGraph->Nodes[i]);
+			if (Node)
+			{
+				const bool bNodeUpdated = Node->RefreshNodeClass();
+				bUpdated = bUpdated || bNodeUpdated;
+			}
+		}
+	}
+
+	return bUpdated;
+}
+
+bool UBehaviorTreeGraphNode_CompositeDecorator::HasErrors() const
+{
+	return bHasObserverError;
+}
+
 void UBehaviorTreeGraphNode_CompositeDecorator::CreateBoundGraph()
 {
 	// Create a new animation graph
diff --git a/Engine/Source/Editor/BehaviorTreeEditor/Private/BehaviorTreeGraphNode_Decorator.cpp b/Engine/Source/Editor/BehaviorTreeEditor/Private/BehaviorTreeGraphNode_Decorator.cpp
--- a/Engine/Source/Editor/BehaviorTreeEditor/Private/BehaviorTreeGraphNode_Decorator.cpp
+++ b/Engine/Source/Editor/BehaviorTreeEditor/Private/BehaviorTreeGraphNode_Decorator.cpp
@@ -20,6 +20,10 @@ FText UBehaviorTreeGraphNode_Decorator::GetNodeTitle(ENodeTitleType::Type TitleT
 	{
 		return FText::FromString(Decorator->GetNodeName());
 	}
+	else if (!ClassData.GetClassName().IsEmpty())
+	{
+		return NSLOCTEXT("BehaviorTreeGraphNode", "UnknownNodeClass", "Can't load class!");
+	}
 
 	return Super::GetNodeTitle(TitleType);
 }
diff --git a/Engine/Source/Editor/BehaviorTreeEditor/Private/BehaviorTreeGraphNode_Service.cpp b/Engine/Source/Editor/BehaviorTreeEditor/Private/BehaviorTreeGraphNode_Service.cpp
--- a/Engine/Source/Editor/BehaviorTreeEditor/Private/BehaviorTreeGraphNode_Service.cpp
+++ b/Engine/Source/Editor/BehaviorTreeEditor/Private/BehaviorTreeGraphNode_Service.cpp
@@ -21,6 +21,10 @@ FText UBehaviorTreeGraphNode_Service::GetNodeTitle(ENodeTitleType::Type TitleTyp
 	{
 		return FText::FromString(Service->GetNodeName());
 	}
+	else if (!ClassData.GetClassName().IsEmpty())
+	{
+		return NSLOCTEXT("BehaviorTreeGraphNode", "UnknownNodeClass", "Can't load class!");
+	}
 
 	return Super::GetNodeTitle(TitleType);
 }
diff --git a/Engine/Source/Editor/BehaviorTreeEditor/Private/BehaviorTreeGraphNode_Task.cpp b/Engine/Source/Editor/BehaviorTreeEditor/Private/BehaviorTreeGraphNode_Task.cpp
--- a/Engine/Source/Editor/BehaviorTreeEditor/Private/BehaviorTreeGraphNode_Task.cpp
+++ b/Engine/Source/Editor/BehaviorTreeEditor/Private/BehaviorTreeGraphNode_Task.cpp
@@ -2,8 +2,6 @@
 
 #include "BehaviorTreeEditorPrivatePCH.h"
 
-#define LOCTEXT_NAMESPACE "BehaviorTreeGraphNode_Task"
-
 UBehaviorTreeGraphNode_Task::UBehaviorTreeGraphNode_Task(const class FPostConstructInitializeProperties& PCIP) : Super(PCIP)
 {
 }
@@ -20,6 +18,10 @@ FText UBehaviorTreeGraphNode_Task::GetNodeTitle(ENodeTitleType::Type TitleType)
 	{
 		return FText::FromString(MyNode->GetNodeName());
 	}
+	else if (!ClassData.GetClassName().IsEmpty())
+	{
+		return NSLOCTEXT("BehaviorTreeGraphNode", "UnknownNodeClass", "Can't load class!");
+	}
 
 	return Super::GetNodeTitle(TitleType);
 }
@@ -33,5 +35,3 @@ void UBehaviorTreeGraphNode_Task::GetContextMenuActions(const FGraphNodeContextM
 {
 	AddContextMenuActionsDecorators(Context);
 }
-
-#undef LOCTEXT_NAMESPACE
diff --git a/Engine/Source/Editor/BehaviorTreeEditor/Private/SGraphNode_BehaviorTree.cpp b/Engine/Source/Editor/BehaviorTreeEditor/Private/SGraphNode_BehaviorTree.cpp
--- a/Engine/Source/Editor/BehaviorTreeEditor/Private/SGraphNode_BehaviorTree.cpp
+++ b/Engine/Source/Editor/BehaviorTreeEditor/Private/SGraphNode_BehaviorTree.cpp
@@ -163,7 +163,7 @@ FSlateColor SGraphNode_BehaviorTree::GetBackgroundColor() const
 		false;
 
 	FLinearColor NodeColor = BehaviorTreeColors::NodeBody::Default;
-	if (BTGraphNode && (BTGraphNode->ErrorMessage.Len() > 0 || BTGraphNode->bHasObserverError))
+	if (BTGraphNode && BTGraphNode->HasErrors())
 	{
 		NodeColor = BehaviorTreeColors::NodeBody::Error;
 	}
