improved lazy loading of blueprint based behavior tree nodes
- missing file

diff --git a/Engine/Source/Editor/BehaviorTreeEditor/Classes/BehaviorTreeGraphNode_Composite.h b/Engine/Source/Editor/BehaviorTreeEditor/Classes/BehaviorTreeGraphNode_Composite.h
--- a/Engine/Source/Editor/BehaviorTreeEditor/Classes/BehaviorTreeGraphNode_Composite.h
+++ b/Engine/Source/Editor/BehaviorTreeEditor/Classes/BehaviorTreeGraphNode_Composite.h
@@ -9,6 +9,7 @@ class UBehaviorTreeGraphNode_Composite : public UBehaviorTreeGraphNode
 	GENERATED_UCLASS_BODY()
 	
 	virtual FText GetNodeTitle(ENodeTitleType::Type TitleType) const OVERRIDE;
+	virtual bool RefreshNodeClass() OVERRIDE{ return false; }
 
 	/** gets icon resource name for title bar */
 	virtual FName GetNameIcon() const OVERRIDE;
