Blackboards are now edited within the Behavior Tree Editor.

Tweaked and reworked the BT editor to allow for multiple modes (like the Blueprint Editor).
If a behavior tree uses a Blackboard, that blackboard is editable in the same asset editor, under "Blackboard mode".
Fixed a number of issues with re-focussing the BT debugger (and some issues that arose when switching BT editor modes). Debugging info should now be properly displayed when selecting AIs, switching modes, starting up the BT editor while playing etc.
Fixed a breakpoint bug where debug pause would be toggled rather than enabled, causing the game to run for a short time.
Added new icons for BT & BB modes.
Added various new guards against nullptr dereferences.
Displayed Actor label rather than FName when choosing actors to debug.

TTP# 337108 - Behaviour Tree Editor: Integrate Blackboard Editor into Behaviour Tree Editor

diff --git a/Engine/Source/Editor/BehaviorTreeEditor/BehaviorTreeEditor.Build.cs b/Engine/Source/Editor/BehaviorTreeEditor/BehaviorTreeEditor.Build.cs
--- a/Engine/Source/Editor/BehaviorTreeEditor/BehaviorTreeEditor.Build.cs
+++ b/Engine/Source/Editor/BehaviorTreeEditor/BehaviorTreeEditor.Build.cs
@@ -49,7 +49,8 @@ public class BehaviorTreeEditor : ModuleRules
 		DynamicallyLoadedModuleNames.AddRange(
             new string[] { 
                 "WorkspaceMenuStructure",
-                "PropertyEditor"
+                "PropertyEditor",
+				"AssetTools",
             }
 		);
 	}
diff --git a/Engine/Source/Editor/BehaviorTreeEditor/Private/AssetTypeActions_BehaviorTree.cpp b/Engine/Source/Editor/BehaviorTreeEditor/Private/AssetTypeActions_BehaviorTree.cpp
--- a/Engine/Source/Editor/BehaviorTreeEditor/Private/AssetTypeActions_BehaviorTree.cpp
+++ b/Engine/Source/Editor/BehaviorTreeEditor/Private/AssetTypeActions_BehaviorTree.cpp
@@ -4,32 +4,48 @@
 #include "Toolkits/AssetEditorManager.h"
 
 #include "BehaviorTree/BehaviorTree.h"
-#include "Editor/BehaviorTreeEditor/Public/BehaviorTreeEditorModule.h"
-#include "Editor/BehaviorTreeEditor/Public/IBehaviorTreeEditor.h"
-
+#include "BehaviorTree/BlackboardData.h"
+#include "BehaviorTreeEditorModule.h"
+#include "BehaviorTreeEditor.h"
 
 #include "AssetTypeActions_BehaviorTree.h"
 #include "SBehaviorTreeDiff.h"
+
 #define LOCTEXT_NAMESPACE "AssetTypeActions"
 
-void FAssetTypeActions_BehaviorTree::GetActions( const TArray<UObject*>& InObjects, FMenuBuilder& MenuBuilder )
+bool FAssetTypeActions_BehaviorTree::BehaviorTreeEditorEnabled()
 {
 	bool BehaviorTreeNewAssetsEnabled=false;
 	GConfig->GetBool(TEXT("BehaviorTreesEd"), TEXT("BehaviorTreeNewAssetsEnabled"), BehaviorTreeNewAssetsEnabled, GEngineIni);
 
 	if (!BehaviorTreeNewAssetsEnabled && !GetDefault<UEditorExperimentalSettings>()->bBehaviorTreeEditor)
+	{
+		return false;
+	}
+
+	return true;
+}
+
+bool FAssetTypeActions_BehaviorTree::HasActions ( const TArray<UObject*>& InObjects ) const
+{
+	return BehaviorTreeEditorEnabled();
+}
+
+void FAssetTypeActions_BehaviorTree::GetActions( const TArray<UObject*>& InObjects, FMenuBuilder& MenuBuilder )
+{
+	if(!BehaviorTreeEditorEnabled())
 	{
 		return;
 	}
 
-	auto Scripts = GetTypedWeakObjectPtrs<UBehaviorTree>(InObjects);
+	auto BehaviorTrees = GetTypedWeakObjectPtrs<UBehaviorTree>(InObjects);
 
 	MenuBuilder.AddMenuEntry(
 		LOCTEXT("BehaviorTree_Edit", "Edit"),
 		LOCTEXT("BehaviorTree_EditTooltip", "Opens the selected Behavior Tree in editor."),
 		FSlateIcon(),
 		FUIAction(
-		FExecuteAction::CreateSP( this, &FAssetTypeActions_BehaviorTree::ExecuteEdit, Scripts ),
+		FExecuteAction::CreateSP( this, &FAssetTypeActions_BehaviorTree::ExecuteEdit, BehaviorTrees ),
 		FCanExecuteAction()
 		)
 	);
@@ -37,35 +53,48 @@ void FAssetTypeActions_BehaviorTree::GetActions( const TArray<UObject*>& InObjec
 
 void FAssetTypeActions_BehaviorTree::OpenAssetEditor( const TArray<UObject*>& InObjects, TSharedPtr<IToolkitHost> EditWithinLevelEditor )
 {
-	bool BehaviorTreeNewAssetsEnabled=false;
-	GConfig->GetBool(TEXT("BehaviorTreesEd"), TEXT("BehaviorTreeNewAssetsEnabled"), BehaviorTreeNewAssetsEnabled, GEngineIni);
-
-	if (!BehaviorTreeNewAssetsEnabled && !GetDefault<UEditorExperimentalSettings>()->bBehaviorTreeEditor)
+	if(!BehaviorTreeEditorEnabled())
 	{
 		return;
 	}
 
 	EToolkitMode::Type Mode = EditWithinLevelEditor.IsValid() ? EToolkitMode::WorldCentric : EToolkitMode::Standalone;
 
-	for (auto ObjIt = InObjects.CreateConstIterator(); ObjIt; ++ObjIt)
+	for(auto Object : InObjects)
 	{
-		auto Script = Cast<UBehaviorTree>(*ObjIt);
-		if (Script != NULL)
+		auto BehaviorTree = Cast<UBehaviorTree>(Object);
+		if(BehaviorTree != nullptr)
 		{
-			FBehaviorTreeEditorModule& BehaviorTreeEditorModule = FModuleManager::LoadModuleChecked<FBehaviorTreeEditorModule>( "BehaviorTreeEditor" );
-			TSharedRef< IBehaviorTreeEditor > NewEditor = BehaviorTreeEditorModule.CreateBehaviorTreeEditor( EToolkitMode::Standalone, EditWithinLevelEditor, Script );
+			// check if we have an editor open for this BT's blackboard & use that if we can
+			bool bFoundExisting = false;
+			if(BehaviorTree->BlackboardAsset != nullptr)
+			{
+				const bool bFocusIfOpen = true;
+				FBehaviorTreeEditor* ExistingInstance = static_cast<FBehaviorTreeEditor*>(FAssetEditorManager::Get().FindEditorForAsset(BehaviorTree->BlackboardAsset, bFocusIfOpen));
+				if(ExistingInstance != nullptr && ExistingInstance->GetBehaviorTree() == nullptr)
+				{
+					ExistingInstance->InitBehaviorTreeEditor(Mode, EditWithinLevelEditor, BehaviorTree);
+					bFoundExisting = true;
+				}
+			}
+			
+			if(!bFoundExisting)
+			{
+				FBehaviorTreeEditorModule& BehaviorTreeEditorModule = FModuleManager::GetModuleChecked<FBehaviorTreeEditorModule>( "BehaviorTreeEditor" );
+				TSharedRef< IBehaviorTreeEditor > NewEditor = BehaviorTreeEditorModule.CreateBehaviorTreeEditor( Mode, EditWithinLevelEditor, BehaviorTree );	
+			}
 		}
 	}
 }
 
 void FAssetTypeActions_BehaviorTree::ExecuteEdit(TArray<TWeakObjectPtr<UBehaviorTree>> Objects)
 {
-	for (auto ObjIt = Objects.CreateConstIterator(); ObjIt; ++ObjIt)
+	for(auto Object : Objects)
 	{
-		auto Object = (*ObjIt).Get();
-		if ( Object )
+		auto BehaviorTree = Object.Get();
+		if(BehaviorTree != NULL)
 		{
-			FAssetEditorManager::Get().OpenEditorForAsset(Object);
+			FAssetEditorManager::Get().OpenEditorForAsset(BehaviorTree);
 		}
 	}
 }
diff --git a/Engine/Source/Editor/BehaviorTreeEditor/Private/AssetTypeActions_BehaviorTree.h b/Engine/Source/Editor/BehaviorTreeEditor/Private/AssetTypeActions_BehaviorTree.h
--- a/Engine/Source/Editor/BehaviorTreeEditor/Private/AssetTypeActions_BehaviorTree.h
+++ b/Engine/Source/Editor/BehaviorTreeEditor/Private/AssetTypeActions_BehaviorTree.h
@@ -8,14 +8,17 @@ class FAssetTypeActions_BehaviorTree : public FAssetTypeActions_Base
 public:
 	// IAssetTypeActions Implementation
 	virtual FText GetName() const OVERRIDE { return NSLOCTEXT("AssetTypeActions", "AssetTypeActions_BehaviorTree", "Behavior Tree"); }
-	virtual FColor GetTypeColor() const OVERRIDE { return FColor(129,50,255); }
+	virtual FColor GetTypeColor() const OVERRIDE { return FColor(149,70,255); }
 	virtual UClass* GetSupportedClass() const OVERRIDE;
-	virtual bool HasActions ( const TArray<UObject*>& InObjects ) const OVERRIDE { return true; }
+	virtual bool HasActions ( const TArray<UObject*>& InObjects ) const OVERRIDE;
 	virtual void GetActions( const TArray<UObject*>& InObjects, FMenuBuilder& MenuBuilder ) OVERRIDE;
 	virtual void OpenAssetEditor( const TArray<UObject*>& InObjects, TSharedPtr<class IToolkitHost> EditWithinLevelEditor = TSharedPtr<IToolkitHost>() ) OVERRIDE;
 	virtual uint32 GetCategories() OVERRIDE { return EAssetTypeCategories::Misc; }
 	virtual void PerformAssetDiff(UObject* Asset1, UObject* Asset2, const struct FRevisionInfo& OldRevision, const struct FRevisionInfo& NewRevision) const OVERRIDE;
 
+	/** Check whether the behavior tree system is enabled */
+	static bool BehaviorTreeEditorEnabled();
+
 private:
 
 	/* Called to open the Behavior Tree defaults view, this opens whatever text dif tool the user has */
diff --git a/Engine/Source/Editor/BehaviorTreeEditor/Private/AssetTypeActions_Blackboard.cpp b/Engine/Source/Editor/BehaviorTreeEditor/Private/AssetTypeActions_Blackboard.cpp
new file mode 100644
--- /dev/null
+++ b/Engine/Source/Editor/BehaviorTreeEditor/Private/AssetTypeActions_Blackboard.cpp
@@ -0,0 +1,77 @@
+// Copyright 1998-2014 Epic Games, Inc. All Rights Reserved.
+
+#include "BehaviorTreeEditorPrivatePCH.h"
+#include "Toolkits/AssetEditorManager.h"
+
+#include "BehaviorTree/BlackboardData.h"
+#include "Editor/BehaviorTreeEditor/Public/BehaviorTreeEditorModule.h"
+#include "Editor/BehaviorTreeEditor/Public/IBehaviorTreeEditor.h"
+
+#include "AssetTypeActions_Blackboard.h"
+#include "AssetTypeActions_BehaviorTree.h"
+
+#define LOCTEXT_NAMESPACE "AssetTypeActions"
+
+UClass* FAssetTypeActions_Blackboard::GetSupportedClass() const
+{
+	return UBlackboardData::StaticClass(); 
+}
+
+bool FAssetTypeActions_Blackboard::HasActions ( const TArray<UObject*>& InObjects ) const
+{
+	return FAssetTypeActions_BehaviorTree::BehaviorTreeEditorEnabled();
+}
+
+void FAssetTypeActions_Blackboard::GetActions( const TArray<UObject*>& InObjects, FMenuBuilder& MenuBuilder )
+{
+	if (!FAssetTypeActions_BehaviorTree::BehaviorTreeEditorEnabled())
+	{
+		return;
+	}
+
+	auto BlackboardData = GetTypedWeakObjectPtrs<UBlackboardData>(InObjects);
+
+	MenuBuilder.AddMenuEntry(
+		LOCTEXT("Blackboard_Edit", "Edit"),
+		LOCTEXT("Blackboard_EditTooltip", "Opens the selected Blackboard in editor."),
+		FSlateIcon(),
+		FUIAction(
+		FExecuteAction::CreateSP( this, &FAssetTypeActions_Blackboard::ExecuteEdit, BlackboardData ),
+		FCanExecuteAction()
+		)
+	);
+}
+
+void FAssetTypeActions_Blackboard::OpenAssetEditor( const TArray<UObject*>& InObjects, TSharedPtr<class IToolkitHost> EditWithinLevelEditor )
+{
+	if (!FAssetTypeActions_BehaviorTree::BehaviorTreeEditorEnabled())
+	{
+		return;
+	}
+
+	EToolkitMode::Type Mode = EditWithinLevelEditor.IsValid() ? EToolkitMode::WorldCentric : EToolkitMode::Standalone;
+
+	for(auto Object : InObjects)
+	{
+		auto BlackboardData = Cast<UBlackboardData>(Object);
+		if(BlackboardData != nullptr)
+		{
+			FBehaviorTreeEditorModule& BehaviorTreeEditorModule = FModuleManager::LoadModuleChecked<FBehaviorTreeEditorModule>( "BehaviorTreeEditor" );
+			BehaviorTreeEditorModule.CreateBehaviorTreeEditor( EToolkitMode::Standalone, EditWithinLevelEditor, BlackboardData );
+		}
+	}	
+}
+
+void FAssetTypeActions_Blackboard::ExecuteEdit(TArray<TWeakObjectPtr<class UBlackboardData>> Objects)
+{
+	for(auto Object : Objects)
+	{
+		auto BlackboardData = Object.Get();
+		if(BlackboardData != nullptr)
+		{
+			FAssetEditorManager::Get().OpenEditorForAsset(BlackboardData);
+		}
+	}
+}
+
+#undef LOCTEXT_NAMESPACE
diff --git a/Engine/Source/Editor/BehaviorTreeEditor/Private/AssetTypeActions_Blackboard.h b/Engine/Source/Editor/BehaviorTreeEditor/Private/AssetTypeActions_Blackboard.h
new file mode 100644
--- /dev/null
+++ b/Engine/Source/Editor/BehaviorTreeEditor/Private/AssetTypeActions_Blackboard.h
@@ -0,0 +1,21 @@
+// Copyright 1998-2014 Epic Games, Inc. All Rights Reserved.
+
+#pragma once
+#include "AssetTypeActions_Base.h"
+
+class FAssetTypeActions_Blackboard : public FAssetTypeActions_Base
+{
+public:
+	// IAssetTypeActions Implementation
+	virtual FText GetName() const override { return NSLOCTEXT("AssetTypeActions", "AssetTypeActions_Blackboard", "Blackboard"); }
+	virtual FColor GetTypeColor() const override { return FColor(201, 29, 85); }
+	virtual UClass* GetSupportedClass() const override;
+	virtual bool HasActions ( const TArray<UObject*>& InObjects ) const override;
+	virtual void GetActions( const TArray<UObject*>& InObjects, FMenuBuilder& MenuBuilder ) override;
+	virtual void OpenAssetEditor( const TArray<UObject*>& InObjects, TSharedPtr<class IToolkitHost> EditWithinLevelEditor = TSharedPtr<IToolkitHost>() ) override;
+	virtual uint32 GetCategories() override { return EAssetTypeCategories::Misc; }
+
+private:
+	/** Handler for when Edit is selected */
+	void ExecuteEdit(TArray<TWeakObjectPtr<class UBlackboardData>> Objects);
+};
\ No newline at end of file
diff --git a/Engine/Source/Editor/BehaviorTreeEditor/Private/BehaviorTreeDebugger.cpp b/Engine/Source/Editor/BehaviorTreeEditor/Private/BehaviorTreeDebugger.cpp
--- a/Engine/Source/Editor/BehaviorTreeEditor/Private/BehaviorTreeDebugger.cpp
+++ b/Engine/Source/Editor/BehaviorTreeEditor/Private/BehaviorTreeDebugger.cpp
@@ -59,17 +59,27 @@ void FBehaviorTreeDebugger::Setup(class UBehaviorTree* InTreeAsset, const class
 	}
 
 #if USE_BEHAVIORTREE_DEBUGGER
-	// check if PIE is already running
 	if (IsPIESimulating())
 	{
 		OnBeginPIE(GEditor->bIsSimulatingInEditor);
 
+		Refresh();
+	}
+#endif
+}
+
+void FBehaviorTreeDebugger::Refresh()
+{
+	if (IsPIESimulating())
+	{
 		// make sure is grabs data if currently paused
 		if (IsPlaySessionPaused() && TreeInstance.IsValid())
 		{
 			FindLockedDebugActor(GEditor->PlayWorld);
 			
 			UpdateDebuggerInstance();
+			UpdateAvailableActions();
+
 			if (DebuggerInstanceIndex != INDEX_NONE)
 			{
 				UpdateDebuggerViewOnStepChange();
@@ -84,7 +94,6 @@ void FBehaviorTreeDebugger::Setup(class UBehaviorTree* InTreeAsset, const class
 			}
 		}
 	}
-#endif
 }
 
 void FBehaviorTreeDebugger::Tick(float DeltaTime)
@@ -165,14 +174,22 @@ bool FBehaviorTreeDebugger::IsTickable() const
 void FBehaviorTreeDebugger::OnBeginPIE(const bool bIsSimulating)
 {
 	bIsPIEActive = true;
-	EditorOwner->RegenerateMenusAndToolbars();
-	EditorOwner->DebuggerUpdateGraph();
+	if(EditorOwner != nullptr)
+	{
+		EditorOwner->RegenerateMenusAndToolbars();
+		EditorOwner->DebuggerUpdateGraph();
+	}
 
 	ActiveBreakpoints.Reset();
 	CollectBreakpointsFromAsset(RootNode.Get());
 
 	FindMatchingTreeInstance();
 
+	// remove these delegates first as we can get multiple calls to OnBeginPIE()
+	USelection::SelectObjectEvent.RemoveAll(this);
+	FBehaviorTreeDelegates::OnTreeStarted.RemoveAll(this);
+	FBehaviorTreeDelegates::OnDebugSelected.RemoveAll(this);
+
 	USelection::SelectObjectEvent.AddRaw(this, &FBehaviorTreeDebugger::OnObjectSelected);
 	FBehaviorTreeDelegates::OnTreeStarted.AddRaw(this, &FBehaviorTreeDebugger::OnTreeStarted);
 	FBehaviorTreeDelegates::OnDebugSelected.AddRaw(this, &FBehaviorTreeDebugger::OnAIDebugSelected);
@@ -181,7 +198,10 @@ void FBehaviorTreeDebugger::OnBeginPIE(const bool bIsSimulating)
 void FBehaviorTreeDebugger::OnEndPIE(const bool bIsSimulating)
 {
 	bIsPIEActive = false;
-	EditorOwner->RegenerateMenusAndToolbars();
+	if(EditorOwner != nullptr)
+	{
+		EditorOwner->RegenerateMenusAndToolbars();
+	}
 
 	USelection::SelectObjectEvent.RemoveAll(this);
 	FBehaviorTreeDelegates::OnTreeStarted.RemoveAll(this);
@@ -680,7 +700,7 @@ void FBehaviorTreeDebugger::FindMatchingTreeInstance()
 		}
 	}
 
-	if (MatchingComp)
+	if (MatchingComp != TreeInstance)
 	{
 		TreeInstance = MatchingComp;
 		UpdateDebuggerViewOnInstanceChange();
@@ -902,7 +922,7 @@ void FBehaviorTreeDebugger::PausePlaySession()
 {
 	if (GUnrealEd->PlayWorld)
 	{
-		GUnrealEd->PlayWorld->bDebugPauseExecution = !GUnrealEd->PlayWorld->bDebugPauseExecution;
+		GUnrealEd->PlayWorld->bDebugPauseExecution = true;
 	}
 }
 
@@ -910,7 +930,7 @@ void FBehaviorTreeDebugger::ResumePlaySession()
 {
 	if (GUnrealEd->PlayWorld)
 	{
-		GUnrealEd->PlayWorld->bDebugPauseExecution = !GUnrealEd->PlayWorld->bDebugPauseExecution;
+		GUnrealEd->PlayWorld->bDebugPauseExecution = false;
 	}
 }
 
@@ -976,6 +996,8 @@ void FBehaviorTreeDebugger::UpdateDebuggerViewOnInstanceChange()
 		}
 
 		DebuggerDetails->UpdateBlackboard(BBAsset);
+
+		Refresh();
 	}
 #endif
 }
@@ -1026,7 +1048,7 @@ FString FBehaviorTreeDebugger::DescribeInstance(class UBehaviorTreeComponent* In
 		AController* TestController = Cast<AController>(InstanceToDescribe->GetOwner());
 		ActorDesc = TestController ?
 			TestController->GetName() :
-			InstanceToDescribe->GetOwner()->GetName();
+			InstanceToDescribe->GetOwner()->GetActorLabel();
 	}
 
 	return ActorDesc;
@@ -1058,6 +1080,8 @@ void FBehaviorTreeDebugger::OnInstanceSelectedInDropdown(class UBehaviorTreeComp
 				SelectedActors->Select(Pawn);
 			}
 		}
+
+		Refresh();
 	}
 }
 
diff --git a/Engine/Source/Editor/BehaviorTreeEditor/Private/BehaviorTreeDebugger.h b/Engine/Source/Editor/BehaviorTreeEditor/Private/BehaviorTreeDebugger.h
--- a/Engine/Source/Editor/BehaviorTreeEditor/Private/BehaviorTreeDebugger.h
+++ b/Engine/Source/Editor/BehaviorTreeEditor/Private/BehaviorTreeDebugger.h
@@ -12,9 +12,13 @@ public:
 
 	virtual void Tick(float DeltaTime) OVERRIDE;
 	virtual bool IsTickable() const OVERRIDE;
+	virtual bool IsTickableWhenPaused() const override { return true; }
 	virtual bool IsTickableInEditor() const OVERRIDE { return true; }
 	virtual TStatId GetStatId() const OVERRIDE { RETURN_QUICK_DECLARE_CYCLE_STAT(FBehaviorTreeEditorTickHelper, STATGROUP_Tickables); }
 
+	/** Refresh the debugging information we are displaying (only when paused, as Tick() updates when running) */
+	void Refresh();
+
 	void Setup(class UBehaviorTree* InTreeAsset, const class FBehaviorTreeEditor* InEditorOwner, TSharedPtr<SBehaviorTreeDebuggerView> DebuggerView);
 
 	bool IsDebuggerReady() const;
diff --git a/Engine/Source/Editor/BehaviorTreeEditor/Private/BehaviorTreeEditor.cpp b/Engine/Source/Editor/BehaviorTreeEditor/Private/BehaviorTreeEditor.cpp
--- a/Engine/Source/Editor/BehaviorTreeEditor/Private/BehaviorTreeEditor.cpp
+++ b/Engine/Source/Editor/BehaviorTreeEditor/Private/BehaviorTreeEditor.cpp
@@ -15,116 +15,16 @@
 #include "Editor/PropertyEditor/Public/PropertyEditorModule.h"
 #include "Editor/PropertyEditor/Public/IDetailsView.h"
 
-#define LOCTEXT_NAMESPACE "BehaviorTreeEditor"
-
-const FName FBehaviorTreeEditor::BTGraphTabId( TEXT( "Document" ) );
-const FName FBehaviorTreeEditor::BTPropertiesTabId( TEXT( "BehaviorTreeEditor_Properties" ) );
-const FName FBehaviorTreeEditor::BTSearchTabId(TEXT("BehaviorTreeEditor_Search"));
-
-//////////////////////////////////////////////////////////////////////////
-struct FBTGraphEditorSummoner : public FDocumentTabFactoryForObjects<UEdGraph>
-{
-public:
-	DECLARE_DELEGATE_RetVal_OneParam(TSharedRef<SGraphEditor>, FOnCreateGraphEditorWidget, UEdGraph*);
-public:
-	FBTGraphEditorSummoner(TSharedPtr<class FBehaviorTreeEditor> InBehaviorTreeEditorPtr, FOnCreateGraphEditorWidget CreateGraphEditorWidgetCallback)
-		: FDocumentTabFactoryForObjects<UEdGraph>(FBehaviorTreeEditor::BTGraphTabId, InBehaviorTreeEditorPtr)
-		, BehaviorTreeEditorPtr(InBehaviorTreeEditorPtr)
-		, OnCreateGraphEditorWidget(CreateGraphEditorWidgetCallback)
-	{
-	}
-
-	virtual void OnTabActivated(TSharedPtr<SDockTab> Tab) const OVERRIDE
-	{
-		TSharedRef<SGraphEditor> GraphEditor = StaticCastSharedRef<SGraphEditor>(Tab->GetContent());
-		BehaviorTreeEditorPtr.Pin()->OnGraphEditorFocused(GraphEditor);
-	}
-
-	virtual void OnTabRefreshed(TSharedPtr<SDockTab> Tab) const OVERRIDE
-	{
-		TSharedRef<SGraphEditor> GraphEditor = StaticCastSharedRef<SGraphEditor>(Tab->GetContent());
-		GraphEditor->NotifyGraphChanged();
-	}
+#include "BehaviorTreeEditorModes.h"
+#include "BehaviorTreeEditorToolbar.h"
+#include "BehaviorTreeEditorTabFactories.h"
+#include "BehaviorTreeEditorCommands.h"
+#include "BehaviorTreeEditorTabs.h"
 
-protected:
-	virtual TAttribute<FText> ConstructTabNameForObject(UEdGraph* DocumentID) const OVERRIDE
-	{
-		return TAttribute<FText>( FText::FromString( DocumentID->GetName() ) );
-	}
-
-	virtual TSharedRef<SWidget> CreateTabBodyForObject(const FWorkflowTabSpawnInfo& Info, UEdGraph* DocumentID) const OVERRIDE
-	{
-		return OnCreateGraphEditorWidget.Execute(DocumentID);
-	}
-
-	virtual const FSlateBrush* GetTabIconForObject(const FWorkflowTabSpawnInfo& Info, UEdGraph* DocumentID) const OVERRIDE
-	{
-		return FEditorStyle::GetBrush("NoBrush");
-	}
-
-protected:
-	TWeakPtr<class FBehaviorTreeEditor> BehaviorTreeEditorPtr;
-	FOnCreateGraphEditorWidget OnCreateGraphEditorWidget;
-};
-
-//////////////////////////////////////////////////////////////////////////
-class FBTCommonCommands : public TCommands<FBTCommonCommands>
-{
-public:
-	/** Constructor */
-	FBTCommonCommands() 
-		: TCommands<FBTCommonCommands>("BTEditor.Common", LOCTEXT("Common", "Common"), NAME_None, FEditorStyle::GetStyleSetName())
-	{
-	}
-
-	TSharedPtr<FUICommandInfo> SearchBT;
-
-	/** Initialize commands */
-	virtual void RegisterCommands() OVERRIDE;
-};
-
-void FBTCommonCommands::RegisterCommands()
-{
-	UI_COMMAND(SearchBT, "Search", "Search this Behavior Tree.", EUserInterfaceActionType::Button, FInputGesture(EModifierKey::Control, EKeys::F));
-}
-
-//////////////////////////////////////////////////////////////////////////
-class FBTDebuggerCommands : public TCommands<FBTDebuggerCommands>
-{
-public:
-	/** Constructor */
-	FBTDebuggerCommands() 
-		: TCommands<FBTDebuggerCommands>("BTEditor.Debugger", LOCTEXT("Debugger", "Debugger"), NAME_None, FEditorStyle::GetStyleSetName())
-	{
-	}
-
-	TSharedPtr<FUICommandInfo> BackInto;
-	TSharedPtr<FUICommandInfo> BackOver;
-	TSharedPtr<FUICommandInfo> ForwardInto;
-	TSharedPtr<FUICommandInfo> ForwardOver;
-	TSharedPtr<FUICommandInfo> StepOut;
-
-	TSharedPtr<FUICommandInfo> PausePlaySession;
-	TSharedPtr<FUICommandInfo> ResumePlaySession;
-	TSharedPtr<FUICommandInfo> StopPlaySession;
-
-	/** Initialize commands */
-	virtual void RegisterCommands() OVERRIDE;
-};
-
-void FBTDebuggerCommands::RegisterCommands()
-{
-	UI_COMMAND(BackInto, "Back: Into", "Show state from previous step, can go into subtrees", EUserInterfaceActionType::Button, FInputGesture());
-	UI_COMMAND(BackOver, "Back: Over", "Show state from previous step, don't go into subtrees", EUserInterfaceActionType::Button, FInputGesture());
-	UI_COMMAND(ForwardInto, "Forward: Into", "Show state from next step, can go into subtrees", EUserInterfaceActionType::Button, FInputGesture());
-	UI_COMMAND(ForwardOver, "Forward: Over", "Show state from next step, don't go into subtrees", EUserInterfaceActionType::Button, FInputGesture());
-	UI_COMMAND(StepOut, "Step Out", "Show state from next step, leave current subtree", EUserInterfaceActionType::Button, FInputGesture());
-	
-	UI_COMMAND(PausePlaySession, "Pause", "Pause simulation", EUserInterfaceActionType::Button, FInputGesture());
-	UI_COMMAND(ResumePlaySession, "Resume", "Resume simulation", EUserInterfaceActionType::Button, FInputGesture() );
-	UI_COMMAND(StopPlaySession, "Stop", "Stop simulation", EUserInterfaceActionType::Button, FInputGesture());
-}
+#define LOCTEXT_NAMESPACE "BehaviorTreeEditor"
 
+const FName FBehaviorTreeEditor::BehaviorTreeMode(TEXT("BehaviorTree"));
+const FName FBehaviorTreeEditor::BlackboardMode(TEXT("Blackboard"));
 
 //////////////////////////////////////////////////////////////////////////
 FBehaviorTreeEditor::FBehaviorTreeEditor() 
@@ -144,6 +44,9 @@ FBehaviorTreeEditor::FBehaviorTreeEditor()
 	bShowDecoratorRangeSelf = false;
 	bSelectedNodeIsInjected = false;
 	SelectedNodesCount = 0;
+
+	BehaviorTree = nullptr;
+	BlackboardData = nullptr;
 }
 
 FBehaviorTreeEditor::~FBehaviorTreeEditor()
@@ -191,134 +94,170 @@ void FBehaviorTreeEditor::PostRedo(bool bSuccess)
 
 void FBehaviorTreeEditor::RegisterTabSpawners(const TSharedRef<class FTabManager>& TabManager)
 {
-	// @todo TabManagement
+	DocumentManager->SetTabManager(TabManager);
+
+	FWorkflowCentricApplication::RegisterTabSpawners(TabManager);
+}
+
+void FBehaviorTreeEditor::RegisterToolbarTab(const TSharedRef<class FTabManager>& TabManager)
+{
 	FAssetEditorToolkit::RegisterTabSpawners(TabManager);
+}
 
-	const IWorkspaceMenuStructure& MenuStructure = WorkspaceMenu::GetMenuStructure();
+void FBehaviorTreeEditor::InitBehaviorTreeEditor( const EToolkitMode::Type Mode, const TSharedPtr< class IToolkitHost >& InitToolkitHost, UObject* InObject )
+{
+	UBehaviorTree* BehaviorTreeToEdit = Cast<UBehaviorTree>(InObject);
+	UBlackboardData* BlackboardDataToEdit = Cast<UBlackboardData>(InObject);
 
-	TabManager->RegisterTabSpawner( BTPropertiesTabId, FOnSpawnTab::CreateSP(this, &FBehaviorTreeEditor::SpawnTab_Properties) )
-		.SetDisplayName( LOCTEXT( "PropertiesTab", "Details" ) )
-		.SetGroup( MenuStructure.GetAssetEditorCategory() );
+	if(BehaviorTreeToEdit != nullptr)
+	{
+		BehaviorTree = BehaviorTreeToEdit;
+		if(BehaviorTree->BlackboardAsset != nullptr)
+		{
+			BlackboardData = BehaviorTree->BlackboardAsset;
+		}
+	}
+	else if(BlackboardDataToEdit != nullptr)
+	{
+		BlackboardData = BlackboardDataToEdit;
+	}
 
-	TabManager->RegisterTabSpawner(BTSearchTabId, FOnSpawnTab::CreateSP(this, &FBehaviorTreeEditor::SpawnTab_Search))
-		.SetDisplayName(LOCTEXT("SearchTab", "Search"))
-		.SetGroup(MenuStructure.GetAssetEditorCategory());
+	TSharedPtr<FBehaviorTreeEditor> ThisPtr(SharedThis(this));
+	if(!DocumentManager.IsValid())
+	{
+		DocumentManager = MakeShareable(new FDocumentTracker);
+		DocumentManager->Initialize(ThisPtr);
 
-	DocumentManager.SetTabManager(TabManager);
-}
+		// Register the document factories
+		{
+			TSharedRef<FDocumentTabFactory> GraphEditorFactory = MakeShareable(new FBTGraphEditorSummoner(ThisPtr,
+				FBTGraphEditorSummoner::FOnCreateGraphEditorWidget::CreateSP(this, &FBehaviorTreeEditor::CreateGraphEditorWidget)
+				));
 
-void FBehaviorTreeEditor::UnregisterTabSpawners(const TSharedRef<class FTabManager>& TabManager)
-{
-	TabManager->UnregisterTabSpawner( BTPropertiesTabId );
+			// Also store off a reference to the grapheditor factory so we can find all the tabs spawned by it later.
+			GraphEditorTabFactoryPtr = GraphEditorFactory;
+			DocumentManager->RegisterDocumentFactory(GraphEditorFactory);
+		}
+	}
 
-	TabManager->UnregisterTabSpawner(BTSearchTabId);
+	TArray<UObject*> ObjectsToEdit;
+	if(BehaviorTree != nullptr)
+	{
+		ObjectsToEdit.Add(BehaviorTree);
+	}
+	if(BlackboardData != nullptr)
+	{
+		ObjectsToEdit.Add(BlackboardData);
+	}
 
-	FAssetEditorToolkit::UnregisterTabSpawners(TabManager);
-}
+	if (!ToolbarBuilder.IsValid())
+	{
+		ToolbarBuilder = MakeShareable(new FBehaviorTreeEditorToolbar(SharedThis(this)));
+	}
 
+	// if we are already editing objects, dont try to recreate the editor from scratch
+	const TArray<UObject*>* EditedObjects = GetObjectsCurrentlyBeingEdited();
+	if(EditedObjects == nullptr || EditedObjects->Num() == 0)
+	{
+		FGraphEditorCommands::Register();
+		FBTCommonCommands::Register();
+		FBTDebuggerCommands::Register();
 
-void FBehaviorTreeEditor::InitBehaviorTreeEditor( const EToolkitMode::Type Mode, const TSharedPtr< class IToolkitHost >& InitToolkitHost, UBehaviorTree* InScript )
-{
-	Script = InScript;
-	check(Script != NULL);
+		const TSharedRef<FTabManager::FLayout> DummyLayout = FTabManager::NewLayout("NullLayout")->AddArea(FTabManager::NewPrimaryArea());
+		const bool bCreateDefaultStandaloneMenu = true;
+		const bool bCreateDefaultToolbar = true;
+		InitAssetEditor( Mode, InitToolkitHost, FBehaviorTreeEditorModule::BehaviorTreeEditorAppIdentifier, DummyLayout, bCreateDefaultStandaloneMenu, bCreateDefaultToolbar, ObjectsToEdit );
 
-	TSharedPtr<FBehaviorTreeEditor> ThisPtr(SharedThis(this));
-	DocumentManager.Initialize(ThisPtr);
-
-	// Register the document factories
-	{
-		TSharedRef<FDocumentTabFactory> GraphEditorFactory = MakeShareable(new FBTGraphEditorSummoner(ThisPtr,
-			FBTGraphEditorSummoner::FOnCreateGraphEditorWidget::CreateSP(this, &FBehaviorTreeEditor::CreateGraphEditorWidget)
-			));
-
-		// Also store off a reference to the grapheditor factory so we can find all the tabs spawned by it later.
-		GraphEditorTabFactoryPtr = GraphEditorFactory;
-		DocumentManager.RegisterDocumentFactory(GraphEditorFactory);
-	}
-
-	FGraphEditorCommands::Register();
-	FBTCommonCommands::Register();
-	FBTDebuggerCommands::Register();
-
-	TSharedRef<FTabManager::FLayout> StandaloneDefaultLayout = FTabManager::NewLayout( "Standalone_BehaviorTree_Layout" )
-	->AddArea
-	(
-		FTabManager::NewPrimaryArea() ->SetOrientation(Orient_Vertical)
-		->Split
-		(
-			FTabManager::NewStack()
-			->SetSizeCoefficient(0.1f)
-			->AddTab(GetToolbarTabId(), ETabState::OpenedTab) 
-			->SetHideTabWell(true) 
-		)
-		->Split
-		(
-			FTabManager::NewSplitter() ->SetOrientation(Orient_Horizontal)
-			->Split
-			(
-				FTabManager::NewStack()
-				->SetSizeCoefficient(0.7f)
-				->AddTab( BTGraphTabId, ETabState::ClosedTab )
-			)
-			->Split
-			(
-				FTabManager::NewStack()
-				->SetSizeCoefficient(0.3f)
-				->AddTab( BTPropertiesTabId, ETabState::OpenedTab )
-				->AddTab( BTSearchTabId, ETabState::ClosedTab )
-			)
-		)
-	);
-
-	const bool bCreateDefaultStandaloneMenu = true;
-	const bool bCreateDefaultToolbar = true;
-	FAssetEditorToolkit::InitAssetEditor( Mode, InitToolkitHost, FBehaviorTreeEditorModule::BehaviorTreeEditorAppIdentifier, StandaloneDefaultLayout, bCreateDefaultStandaloneMenu, bCreateDefaultToolbar, Script );
+		BindCommonCommands();
+		ExtendMenu();
+		CreateInternalWidgets();
+
+		Debugger = MakeShareable(new FBehaviorTreeDebugger);
+		Debugger->Setup(BehaviorTree, this, DebuggerView);
+		BindDebuggerToolbarCommands();
+
+		FBehaviorTreeEditorModule& BehaviorTreeEditorModule = FModuleManager::GetModuleChecked<FBehaviorTreeEditorModule>( "BehaviorTreeEditor" );
+		AddMenuExtender(BehaviorTreeEditorModule.GetMenuExtensibilityManager()->GetAllExtenders(GetToolkitCommands(), GetEditingObjects()));
+
+		AddApplicationMode(BehaviorTreeMode, MakeShareable(new FBehaviorTreeEditorApplicationMode(SharedThis(this))));
+		AddApplicationMode(BlackboardMode, MakeShareable(new FBlackboardEditorApplicationMode(SharedThis(this))));
+	}
+	else
+	{
+		check(Debugger.IsValid());
+		Debugger->Setup(BehaviorTree, this, DebuggerView);
+	}
+
+	if(BehaviorTreeToEdit != nullptr)
+	{
+		SetCurrentMode(BehaviorTreeMode);
+	}
+	else if(BlackboardDataToEdit != nullptr)
+	{
+		SetCurrentMode(BlackboardMode);
+	}
+
+	RegenerateMenusAndToolbars();
+}
 
+void FBehaviorTreeEditor::RestoreBehaviorTree()
+{
 	// Update BT asset data based on saved graph to have correct data in editor
+	UBehaviorTreeGraph* MyGraph = Cast<UBehaviorTreeGraph>(BehaviorTree->BTGraph);
+	const bool bNewGraph = MyGraph == NULL;
+	if (MyGraph == NULL)
 	{
-		UBehaviorTreeGraph* MyGraph = Cast<UBehaviorTreeGraph>(Script->BTGraph);
-		if (MyGraph == NULL)
-		{
-			Script->BTGraph = FBlueprintEditorUtils::CreateNewGraph(Script, TEXT("Behavior Tree"), UBehaviorTreeGraph::StaticClass(), UEdGraphSchema_BehaviorTree::StaticClass());
-			MyGraph = Cast<UBehaviorTreeGraph>(Script->BTGraph);
-			check(MyGraph);
+		BehaviorTree->BTGraph = FBlueprintEditorUtils::CreateNewGraph(BehaviorTree, TEXT("Behavior Tree"), UBehaviorTreeGraph::StaticClass(), UEdGraphSchema_BehaviorTree::StaticClass());
+		MyGraph = Cast<UBehaviorTreeGraph>(BehaviorTree->BTGraph);
+		check(MyGraph);
 
-			// Initialize the behavior tree graph
-			const UEdGraphSchema* Schema = MyGraph->GetSchema();
-			Schema->CreateDefaultNodesForGraph(*MyGraph);
-		}
-		else
-		{
-			MyGraph->UpdatePinConnectionTypes();
-			MyGraph->UpdateDeprecatedNodes();
-		}
+		// Initialize the behavior tree graph
+		const UEdGraphSchema* Schema = MyGraph->GetSchema();
+		Schema->CreateDefaultNodesForGraph(*MyGraph);
+	}
+	else
+	{
+		MyGraph->UpdatePinConnectionTypes();
+		MyGraph->UpdateDeprecatedNodes();
+	}
 
-		MyGraph->UpdateBlackboardChange();
-		MyGraph->UpdateInjectedNodes();
-		MyGraph->UpdateUnknownNodeClasses();
+	MyGraph->UpdateBlackboardChange();
+	MyGraph->UpdateInjectedNodes();
+	MyGraph->UpdateUnknownNodeClasses();
 
-		TSharedRef<FTabPayload_UObject> Payload = FTabPayload_UObject::Make(MyGraph);
-		DocumentManager.OpenDocument(Payload, FDocumentTracker::OpenNewDocument);
+	TSharedRef<FTabPayload_UObject> Payload = FTabPayload_UObject::Make(MyGraph);
+	TSharedPtr<SDockTab> DocumentTab = DocumentManager->OpenDocument(Payload, bNewGraph ? FDocumentTracker::OpenNewDocument : FDocumentTracker::RestorePreviousDocument);
 
+	if(BehaviorTree->LastEditedDocuments.Num() > 0)
+	{
+		TSharedRef<SGraphEditor> GraphEditor = StaticCastSharedRef<SGraphEditor>(DocumentTab->GetContent());
+		GraphEditor->SetViewLocation(BehaviorTree->LastEditedDocuments[0].SavedViewOffset, BehaviorTree->LastEditedDocuments[0].SavedZoomAmount);
+	}
+
+	if(bNewGraph)
+	{
 		const bool bIncreaseVersionNum = false;
 		MyGraph->UpdateAsset(UBehaviorTreeGraph::ClearDebuggerFlags, bIncreaseVersionNum);
-
-		FAbortDrawHelper EmptyMode;
-		bShowDecoratorRangeLower = false;
-		bShowDecoratorRangeSelf = false;
-		bSelectedNodeIsInjected = false;
-		MyGraph->UpdateAbortHighlight(EmptyMode, EmptyMode);
+	}
+	else
+	{
+		RefreshDebugger();
 	}
 
-	BindCommonCommands();
-	ExtendMenu();
+	FAbortDrawHelper EmptyMode;
+	bShowDecoratorRangeLower = false;
+	bShowDecoratorRangeSelf = false;
+	bSelectedNodeIsInjected = false;
+	MyGraph->UpdateAbortHighlight(EmptyMode, EmptyMode);
+}
 
-	Debugger = MakeShareable(new FBehaviorTreeDebugger);
-	Debugger->Setup(Script, this, DebuggerView);
-	BindDebuggerToolbarCommands();
+void FBehaviorTreeEditor::SaveEditedObjectState()
+{
+	// Clear currently edited documents
+	BehaviorTree->LastEditedDocuments.Empty();
 
-	ExtentToolbar();
-	RegenerateMenusAndToolbars();
+	// Ask all open documents to save their state, which will update LastEditedDocuments
+	DocumentManager->SaveAllState();
 }
 
 bool FBehaviorTreeEditor::IsDebuggerReady() const
@@ -514,107 +453,89 @@ bool FBehaviorTreeEditor::InEditingMode(bool bGraphIsEditable) const
 	return bGraphIsEditable && FBehaviorTreeDebugger::IsPIENotSimulating();
 }
 
-TSharedRef<SDockTab> FBehaviorTreeEditor::SpawnTab_Search(const FSpawnTabArgs& Args)
+TSharedRef<SWidget> FBehaviorTreeEditor::SpawnSearch()
 {
 	FindResults = SNew(SFindInBT, SharedThis(this));
-
-	TSharedRef<SDockTab> SpawnedTab = SNew(SDockTab)
-		.Icon( FEditorStyle::GetBrush("SoundClassEditor.Tabs.Properties"))
-		.Label(LOCTEXT("BehaviorTreeEditorSearchTitle", "Search"))
-		[
-			FindResults.ToSharedRef()
-		];
-
-	return SpawnedTab;
+	return FindResults.ToSharedRef();
 }
 
-TSharedRef<SDockTab> FBehaviorTreeEditor::SpawnTab_Properties(const FSpawnTabArgs& Args)
+TSharedRef<SWidget> FBehaviorTreeEditor::SpawnProperties()
 {
-	check( Args.GetTabId() == BTPropertiesTabId );
-
-	CreateInternalWidgets();
-
-	TSharedRef<SDockTab> SpawnedTab = SNew(SDockTab)
-		.Icon( FEditorStyle::GetBrush("SoundClassEditor.Tabs.Properties") )
-		.Label( LOCTEXT( "SoundClassPropertiesTitle", "Details" ) )
+	return
+		SNew(SVerticalBox)
+		+SVerticalBox::Slot()
+		.AutoHeight()
+		.HAlign(HAlign_Fill)
+		.VAlign(VAlign_Top)
+		[
+			DetailsView.ToSharedRef()
+		]
+		+SVerticalBox::Slot()
+		.AutoHeight()
+		.HAlign(HAlign_Left)
+		.VAlign(VAlign_Top)
 		[
 			SNew(SVerticalBox)
 			+SVerticalBox::Slot()
-			.AutoHeight()
 			.HAlign(HAlign_Fill)
-			.VAlign(VAlign_Top)
-			[
-				DetailsView.ToSharedRef()
-			]
-			+SVerticalBox::Slot()
-			.AutoHeight()
-			.HAlign(HAlign_Left)
-			.VAlign(VAlign_Top)
+			.Padding(0.0f, 5.0f)
 			[
-				SNew(SVerticalBox)
-				+SVerticalBox::Slot()
-				.HAlign(HAlign_Fill)
-				.Padding(0.0f, 5.0f)
-				[
-					SNew(SBorder)
-					.BorderBackgroundColor(BehaviorTreeColors::NodeBody::InjectedSubNode)
-					.BorderImage(FEditorStyle::GetBrush("Graph.StateNode.Body"))
-					.Visibility(this, &FBehaviorTreeEditor::GetInjectedNodeVisibility)
-					.Padding(FMargin(5.0f))
-					[
-						SNew(STextBlock)
-						.Text(LOCTEXT("InjectedNode", "Node is injected by subtree and can't be edited"))
-					]
-				]
-				+SVerticalBox::Slot()
-				.HAlign(HAlign_Fill)
-				.Padding(0.0f, 5.0f)
+				SNew(SBorder)
+				.BorderBackgroundColor(BehaviorTreeColors::NodeBody::InjectedSubNode)
+				.BorderImage(FEditorStyle::GetBrush("Graph.StateNode.Body"))
+				.Visibility(this, &FBehaviorTreeEditor::GetInjectedNodeVisibility)
+				.Padding(FMargin(5.0f))
 				[
-					SNew(SBorder)
-					.BorderBackgroundColor(BehaviorTreeColors::NodeBorder::HighlightAbortRange0)
-					.BorderImage( FEditorStyle::GetBrush( "Graph.StateNode.Body" ) )
-					.Visibility(this, &FBehaviorTreeEditor::GetRangeLowerVisibility)
-					.Padding(FMargin(5.0f))
-					[
-						SNew(STextBlock)
-						.Text(FText::Format(LOCTEXT("AbortModeHighlight", "Nodes aborted by mode: {0}"), LOCTEXT("AbortPriorityLower", "Lower Priority")))
-					]
-				]
-				+SVerticalBox::Slot()
-				.HAlign(HAlign_Fill)
-				.Padding(0.0f, 5.0f)
-				[
-					SNew(SBorder)
-					.BorderBackgroundColor(BehaviorTreeColors::NodeBorder::HighlightAbortRange1)
-					.BorderImage( FEditorStyle::GetBrush( "Graph.StateNode.Body" ) )
-					.Visibility(this, &FBehaviorTreeEditor::GetRangeSelfVisibility)
-					.Padding(FMargin(5.0f))
-					[
-						SNew(STextBlock)
-						.Text(FText::Format(LOCTEXT("AbortModeHighlight", "Nodes aborted by mode: {0}"), LOCTEXT("AbortPrioritySelf", "Self")))
-					]
+					SNew(STextBlock)
+					.Text(LOCTEXT("InjectedNode", "Node is injected by subtree and can't be edited"))
 				]
 			]
 			+SVerticalBox::Slot()
-			.FillHeight(1.0f)
 			.HAlign(HAlign_Fill)
-			.VAlign(VAlign_Bottom)
-			.Padding(5.0f)
+			.Padding(0.0f, 5.0f)
 			[
-				SNew(SSeparator)
-				.Orientation( Orient_Horizontal )
-				.Visibility(this, &FBehaviorTreeEditor::GetDebuggerDetailsVisibility)
+				SNew(SBorder)
+				.BorderBackgroundColor(BehaviorTreeColors::NodeBorder::HighlightAbortRange0)
+				.BorderImage( FEditorStyle::GetBrush( "Graph.StateNode.Body" ) )
+				.Visibility(this, &FBehaviorTreeEditor::GetRangeLowerVisibility)
+				.Padding(FMargin(5.0f))
+				[
+					SNew(STextBlock)
+					.Text(FText::Format(LOCTEXT("AbortModeHighlight", "Nodes aborted by mode: {0}"), LOCTEXT("AbortPriorityLower", "Lower Priority")))
+				]
 			]
 			+SVerticalBox::Slot()
-			.AutoHeight()
-			.HAlign(HAlign_Left)
-			.VAlign(VAlign_Top)
+			.HAlign(HAlign_Fill)
+			.Padding(0.0f, 5.0f)
 			[
-				DebuggerView.ToSharedRef()
+				SNew(SBorder)
+				.BorderBackgroundColor(BehaviorTreeColors::NodeBorder::HighlightAbortRange1)
+				.BorderImage( FEditorStyle::GetBrush( "Graph.StateNode.Body" ) )
+				.Visibility(this, &FBehaviorTreeEditor::GetRangeSelfVisibility)
+				.Padding(FMargin(5.0f))
+				[
+					SNew(STextBlock)
+					.Text(FText::Format(LOCTEXT("AbortModeHighlight", "Nodes aborted by mode: {0}"), LOCTEXT("AbortPrioritySelf", "Self")))
+				]
 			]
+		]
+		+SVerticalBox::Slot()
+		.FillHeight(1.0f)
+		.HAlign(HAlign_Fill)
+		.VAlign(VAlign_Bottom)
+		.Padding(5.0f)
+		[
+			SNew(SSeparator)
+			.Orientation( Orient_Horizontal )
+			.Visibility(this, &FBehaviorTreeEditor::GetDebuggerDetailsVisibility)
+		]
+		+SVerticalBox::Slot()
+		.AutoHeight()
+		.HAlign(HAlign_Left)
+		.VAlign(VAlign_Top)
+		[
+			DebuggerView.ToSharedRef()
 		];
-
-	return SpawnedTab;
 }
 
 void FBehaviorTreeEditor::OnSelectedNodesChanged(const TSet<class UObject*>& NewSelection)
@@ -679,7 +600,7 @@ void FBehaviorTreeEditor::OnSelectedNodesChanged(const TSet<class UObject*>& New
 		}
 	}
 
-	UBehaviorTreeGraph* MyGraph = Cast<UBehaviorTreeGraph>(Script->BTGraph);
+	UBehaviorTreeGraph* MyGraph = Cast<UBehaviorTreeGraph>(BehaviorTree->BTGraph);
 	FAbortDrawHelper Mode0, Mode1;
 	bShowDecoratorRangeLower = false;
 	bShowDecoratorRangeSelf = false;
@@ -874,7 +795,7 @@ void FBehaviorTreeEditor::BindCommonCommands()
 
 void FBehaviorTreeEditor::SearchTree()
 {
-	TabManager->InvokeTab(BTSearchTabId);
+	TabManager->InvokeTab(FBehaviorTreeEditorTabs::SearchID);
 	FindResults->FocusForUse();
 }
 
@@ -883,57 +804,6 @@ bool FBehaviorTreeEditor::CanSearchTree() const
 	return true;
 }
 
-void FBehaviorTreeEditor::ExtentToolbar()
-{
-	struct Local
-	{
-		static void FillToolbar(FToolBarBuilder& ToolbarBuilder, FBehaviorTreeEditor* EditorOwner, TSharedRef<SWidget> SelectionBox)
-		{
-			const bool bCanShowDebugger = EditorOwner && EditorOwner->IsDebuggerReady();
-			if (bCanShowDebugger)
-			{
-				ToolbarBuilder.BeginSection("CachedState");
-				{
-					ToolbarBuilder.AddToolBarButton(FBTDebuggerCommands::Get().BackOver);
-					ToolbarBuilder.AddToolBarButton(FBTDebuggerCommands::Get().BackInto);
-					ToolbarBuilder.AddToolBarButton(FBTDebuggerCommands::Get().ForwardInto);
-					ToolbarBuilder.AddToolBarButton(FBTDebuggerCommands::Get().ForwardOver);
-					ToolbarBuilder.AddToolBarButton(FBTDebuggerCommands::Get().StepOut);
-				}
-				ToolbarBuilder.EndSection();
-				ToolbarBuilder.BeginSection("World");
-				{
-					ToolbarBuilder.AddToolBarButton(FBTDebuggerCommands::Get().PausePlaySession);
-					ToolbarBuilder.AddToolBarButton(FBTDebuggerCommands::Get().ResumePlaySession);
-					ToolbarBuilder.AddToolBarButton(FBTDebuggerCommands::Get().StopPlaySession);
-					ToolbarBuilder.AddSeparator();
-					ToolbarBuilder.AddWidget(SelectionBox);
-				}
-				ToolbarBuilder.EndSection();
-			}
-		}
-	};
-
-	TSharedPtr<FExtender> ToolbarExtender = MakeShareable(new FExtender);
-	TSharedRef<SWidget> SelectionBox = SNew(SComboButton)
-		.OnGetMenuContent( this, &FBehaviorTreeEditor::OnGetDebuggerActorsMenu )
-		.ButtonContent()
-		[
-			SNew(STextBlock)
-			.ToolTipText( LOCTEXT("SelectDebugActor", "Pick actor to debug") )
-			.Text(this, &FBehaviorTreeEditor::GetDebuggerActorDesc )
-		];
-
-	ToolbarExtender->AddToolBarExtension("Asset", EExtensionHook::After, GetToolkitCommands(),
-		FToolBarExtensionDelegate::CreateStatic( &Local::FillToolbar, this, SelectionBox )
-		);
-
-	AddToolbarExtender(ToolbarExtender);
-
-	FBehaviorTreeEditorModule& BehaviorTreeEditorModule = FModuleManager::LoadModuleChecked<FBehaviorTreeEditorModule>( "BehaviorTreeEditor" );
-	AddMenuExtender(BehaviorTreeEditorModule.GetMenuExtensibilityManager()->GetAllExtenders(GetToolkitCommands(), GetEditingObjects()));
-}
-
 TSharedRef<class SWidget> FBehaviorTreeEditor::OnGetDebuggerActorsMenu()
 {
 	FMenuBuilder MenuBuilder(true, NULL);
@@ -1069,7 +939,7 @@ void FBehaviorTreeEditor::OnFinishedChangingProperties(const FPropertyChangedEve
 					FAbortDrawHelper Mode0, Mode1;
 					GetAbortModePreview(Cast<UBTDecorator>(DecoratorNode->NodeInstance), Mode0, Mode1);
 
-					UBehaviorTreeGraph* MyGraph = Cast<UBehaviorTreeGraph>(Script->BTGraph);
+					UBehaviorTreeGraph* MyGraph = Cast<UBehaviorTreeGraph>(BehaviorTree->BTGraph);
 					MyGraph->UpdateAbortHighlight(Mode0, Mode1);
 				}
 			}			
@@ -1078,7 +948,7 @@ void FBehaviorTreeEditor::OnFinishedChangingProperties(const FPropertyChangedEve
 
 	if (PropertyChangedEvent.Property && PropertyChangedEvent.Property->GetFName() == TEXT("BehaviorAsset"))
 	{
-		UBehaviorTreeGraph* MyGraph = Cast<UBehaviorTreeGraph>(Script->BTGraph);
+		UBehaviorTreeGraph* MyGraph = Cast<UBehaviorTreeGraph>(BehaviorTree->BTGraph);
 		MyGraph->UpdateInjectedNodes();
 		MyGraph->UpdateAsset(UBehaviorTreeGraph::ClearDebuggerFlags);
 	}
@@ -1086,7 +956,7 @@ void FBehaviorTreeEditor::OnFinishedChangingProperties(const FPropertyChangedEve
 
 void FBehaviorTreeEditor::OnPackageSaved(const FString& PackageFileName, UObject* Outer)
 {
-	UBehaviorTreeGraph* MyGraph = Script ? Cast<UBehaviorTreeGraph>(Script->BTGraph) : NULL;
+	UBehaviorTreeGraph* MyGraph = BehaviorTree ? Cast<UBehaviorTreeGraph>(BehaviorTree->BTGraph) : NULL;
 	if (MyGraph)
 	{
 		const bool bUpdated = MyGraph->UpdateInjectedNodes();
@@ -1099,7 +969,7 @@ void FBehaviorTreeEditor::OnPackageSaved(const FString& PackageFileName, UObject
 
 void FBehaviorTreeEditor::OnClassListUpdated()
 {
-	UBehaviorTreeGraph* MyGraph = Script ? Cast<UBehaviorTreeGraph>(Script->BTGraph) : NULL;
+	UBehaviorTreeGraph* MyGraph = BehaviorTree ? Cast<UBehaviorTreeGraph>(BehaviorTree->BTGraph) : NULL;
 	if (MyGraph)
 	{
 		const bool bUpdated = MyGraph->UpdateUnknownNodeClasses();
@@ -1416,8 +1286,8 @@ void FBehaviorTreeEditor::PasteNodesHere(const FVector2D& Location)
 	// Update UI
 	CurrentGraphEditor->NotifyGraphChanged();
 
-	Script->PostEditChange();
-	Script->MarkPackageDirty();
+	BehaviorTree->PostEditChange();
+	BehaviorTree->MarkPackageDirty();
 }
 
 bool FBehaviorTreeEditor::CanPasteNodes() const
@@ -1489,15 +1359,15 @@ void FBehaviorTreeEditor::OnNodeDoubleClicked(class UEdGraphNode* Node)
 			TSharedRef<FTabPayload_UObject> Payload = FTabPayload_UObject::Make(const_cast<UEdGraph*>(Decorator->GetBoundGraph()));
 			
 			TArray< TSharedPtr<SDockTab> > MatchingTabs;
-			DocumentManager.FindMatchingTabs(Payload, MatchingTabs);
+			DocumentManager->FindMatchingTabs(Payload, MatchingTabs);
 			if (MatchingTabs.Num())
 			{
-				DocumentManager.CloseTab(Payload);
-				DocumentManager.OpenDocument(Payload, FDocumentTracker::RestorePreviousDocument);
+				DocumentManager->CloseTab(Payload);
+				DocumentManager->OpenDocument(Payload, FDocumentTracker::RestorePreviousDocument);
 			}
 			else
 			{
-				DocumentManager.OpenDocument(Payload, FDocumentTracker::OpenNewDocument);
+				DocumentManager->OpenDocument(Payload, FDocumentTracker::OpenNewDocument);
 			}
 		}
 	}
@@ -1635,9 +1505,9 @@ void FBehaviorTreeEditor::OnGraphEditorFocused(const TSharedRef<SGraphEditor>& I
 
 void FBehaviorTreeEditor::SaveAsset_Execute()
 {
-	if (Script)
+	if (BehaviorTree)
 	{
-		UBehaviorTreeGraph* BTGraph = Cast<UBehaviorTreeGraph>(Script->BTGraph);
+		UBehaviorTreeGraph* BTGraph = Cast<UBehaviorTreeGraph>(BehaviorTree->BTGraph);
 		if (BTGraph)
 		{
 			BTGraph->UpdateAsset(UBehaviorTreeGraph::SkipDebuggerFlags);
@@ -1806,7 +1676,7 @@ bool FBehaviorTreeEditor::CanToggleBreakpoint() const
 
 void FBehaviorTreeEditor::JumpToNode(const UEdGraphNode* Node)
 {
-	TSharedPtr<SDockTab> ActiveTab = DocumentManager.GetActiveTab();
+	TSharedPtr<SDockTab> ActiveTab = DocumentManager->GetActiveTab();
 	if (ActiveTab.IsValid())
 	{
 		TSharedPtr<SGraphEditor> GraphEditor = StaticCastSharedRef<SGraphEditor>(ActiveTab->GetContent());
@@ -1822,9 +1692,35 @@ TWeakPtr<SGraphEditor> FBehaviorTreeEditor::GetFocusedGraphPtr() const
 	return FocusedGraphEdPtr;
 }
 
+bool FBehaviorTreeEditor::CanAccessBlackboardMode() const
+{
+	return GetBlackboardData() != nullptr;
+}
+
+bool FBehaviorTreeEditor::CanAccessBehaviorTreeMode() const
+{
+	return BehaviorTree != nullptr;
+}
+
+FText FBehaviorTreeEditor::GetLocalizedMode(FName InMode)
+{
+	static TMap< FName, FText > LocModes;
+
+	if (LocModes.Num() == 0)
+	{
+		LocModes.Add( BehaviorTreeMode, LOCTEXT("BehaviorTreeMode", "Behavior Tree") );
+		LocModes.Add( BlackboardMode, LOCTEXT("BlackboardMode", "Blackboard") );
+	}
+
+	check( InMode != NAME_None );
+	const FText* OutDesc = LocModes.Find( InMode );
+	check( OutDesc );
+	return *OutDesc;
+}
+
 UEdGraphNode* FBehaviorTreeEditor::FindInjectedNode(int32 Index) const
 {
-	UBehaviorTreeGraph* BTGraph = Cast<UBehaviorTreeGraph>(Script->BTGraph);
+	UBehaviorTreeGraph* BTGraph = Cast<UBehaviorTreeGraph>(BehaviorTree->BTGraph);
 	return BTGraph ? BTGraph->FindInjectedNode(Index) : NULL;
 }
 
@@ -1841,6 +1737,18 @@ void FBehaviorTreeEditor::DoubleClickNode(class UEdGraphNode* Node)
 	OnNodeDoubleClicked(Node);
 }
 
+void FBehaviorTreeEditor::FocusWindow(UObject* ObjectToFocusOn)
+{
+	if(ObjectToFocusOn == BehaviorTree)
+	{
+		SetCurrentMode(BehaviorTreeMode);
+	}
+	else if(ObjectToFocusOn == GetBlackboardData())
+	{
+		SetCurrentMode(BlackboardMode);
+	}
+}
+
 void FBehaviorTreeEditor::InitializeDebuggerState(class FBehaviorTreeDebugger* ParentDebugger) const
 {
 	if (Debugger.IsValid())
@@ -1865,11 +1773,37 @@ void FBehaviorTreeEditor::DebuggerSwitchAsset(UBehaviorTree* NewAsset)
 
 void FBehaviorTreeEditor::DebuggerUpdateGraph()
 {
-	UBehaviorTreeGraph* BTGraph = Script ? Cast<UBehaviorTreeGraph>(Script->BTGraph) : NULL;
+	UBehaviorTreeGraph* BTGraph = BehaviorTree ? Cast<UBehaviorTreeGraph>(BehaviorTree->BTGraph) : NULL;
 	if (BTGraph)
 	{
 		BTGraph->RebuildExecutionOrder();
 	}
 }
 
+FText FBehaviorTreeEditor::GetToolkitName() const
+{
+	const UObject* EditingObject = GetCurrentMode() == BehaviorTreeMode ? (UObject*)BehaviorTree : (UObject*)GetBlackboardData();
+	if(EditingObject != nullptr)
+	{
+		return FAssetEditorToolkit::GetDescriptionForObject(EditingObject);
+	}
+
+	return FText();
+}
+
+UBehaviorTree* FBehaviorTreeEditor::GetBehaviorTree() const 
+{
+	return BehaviorTree; 
+}
+
+UBlackboardData* FBehaviorTreeEditor::GetBlackboardData() const 
+{
+	return BehaviorTree == nullptr ? BlackboardData : BehaviorTree->BlackboardAsset; 
+}
+
+void FBehaviorTreeEditor::RefreshDebugger()
+{
+	Debugger->Refresh();
+}
+
 #undef LOCTEXT_NAMESPACE
diff --git a/Engine/Source/Editor/BehaviorTreeEditor/Private/BehaviorTreeEditor.h b/Engine/Source/Editor/BehaviorTreeEditor/Private/BehaviorTreeEditor.h
--- a/Engine/Source/Editor/BehaviorTreeEditor/Private/BehaviorTreeEditor.h
+++ b/Engine/Source/Editor/BehaviorTreeEditor/Private/BehaviorTreeEditor.h
@@ -13,15 +13,15 @@ public:
 	virtual ~FBehaviorTreeEditor();
 
 	virtual void RegisterTabSpawners(const TSharedRef<class FTabManager>& TabManager) OVERRIDE;
-	virtual void UnregisterTabSpawners(const TSharedRef<class FTabManager>& TabManager) OVERRIDE;
 
-	void InitBehaviorTreeEditor( const EToolkitMode::Type Mode, const TSharedPtr< class IToolkitHost >& InitToolkitHost, UBehaviorTree* Script );
+	void InitBehaviorTreeEditor( const EToolkitMode::Type Mode, const TSharedPtr< class IToolkitHost >& InitToolkitHost, UObject* InObject );
 
 	// Begin IToolkit interface
 	virtual FName GetToolkitFName() const OVERRIDE;
 	virtual FText GetBaseToolkitName() const OVERRIDE;
 	virtual FString GetWorldCentricTabPrefix() const OVERRIDE;
 	virtual FLinearColor GetWorldCentricTabColorScale() const OVERRIDE;
+	virtual FText GetToolkitName() const OVERRIDE;
 	// End IToolkit interface
 
 	// Begin IBehaviorTreeEditor interface
@@ -29,6 +29,7 @@ public:
 	virtual void InitializeDebuggerState(class FBehaviorTreeDebugger* ParentDebugger) const OVERRIDE;
 	virtual UEdGraphNode* FindInjectedNode(int32 Index) const OVERRIDE;
 	virtual void DoubleClickNode(class UEdGraphNode* Node) OVERRIDE;
+	virtual void FocusWindow(UObject* ObjectToFocusOn = NULL) OVERRIDE;
 	// End IBehaviorTreeEditor interface
 
 	// Begin FEditorUndoClient Interface
@@ -102,6 +103,43 @@ public:
 
 	TWeakPtr<SGraphEditor> GetFocusedGraphPtr() const;
 
+	/** Check whether the behavior tree mode can be accessed (i.e whether we have a valid tree to edit) */
+	bool CanAccessBehaviorTreeMode() const;
+
+	/** Check whether the blackboard mode can be accessed (i.e whether we have a valid blackboard to edit) */
+	bool CanAccessBlackboardMode() const;
+
+	/** 
+	 * Get the localized text to display for the specified mode 
+	 * @param	InMode	The mode to display
+	 * @return the localized text representation of the mode
+	 */
+	static FText GetLocalizedMode(FName InMode);
+
+	/** Access the toolbar builder for this editor */
+	TSharedPtr<class FBehaviorTreeEditorToolbar> GetToolbarBuilder() { return ToolbarBuilder; }
+
+	/** Get the behavior tree we are editing (if any) */
+	UBehaviorTree* GetBehaviorTree() const;
+
+	/** Get the blackboard we are editing (if any) */
+	UBlackboardData* GetBlackboardData() const;
+
+	/** Spawns the tab with the update graph inside */
+	TSharedRef<SWidget> SpawnProperties();
+
+	/** Spawns the search tab */
+	TSharedRef<SWidget> SpawnSearch();
+
+	// @todo This is a hack for now until we reconcile the default toolbar with application modes [duplicated from counterpart in Blueprint Editor]
+	void RegisterToolbarTab(const TSharedRef<class FTabManager>& TabManager);
+
+	/** Restores the behavior tree graph we were editing or creates a new one if none is available */
+	void RestoreBehaviorTree();
+
+	/** Save the graph state for later editing */
+	void SaveEditedObjectState();
+
 protected:
 	/** Called when "Save" is clicked for this asset */
 	virtual void SaveAsset_Execute() OVERRIDE;
@@ -119,9 +157,6 @@ private:
 	/** Setup common commands */
 	void BindCommonCommands();
 
-	/** Add custom widgets to toolbar */
-	void ExtentToolbar();
-
 	/** Setup commands */
 	void BindDebuggerToolbarCommands();
 
@@ -134,17 +169,17 @@ private:
 	/** prepare range of nodes that can be aborted by this decorator */
 	void GetAbortModePreview(const class UBTDecorator* DecoratorOb, struct FAbortDrawHelper& Mode0, struct FAbortDrawHelper& Mode1);
 
-	/** Spawns the tab with the update graph inside */
-	TSharedRef<SDockTab> SpawnTab_Properties(const FSpawnTabArgs& Args);
+	/** Refresh the debugger's display */
+	void RefreshDebugger();
 
-	/** Spawns the search tab */
-	TSharedRef<SDockTab> SpawnTab_Search(const FSpawnTabArgs& Args);
-
-	FDocumentTracker DocumentManager;
+	TSharedPtr<FDocumentTracker> DocumentManager;
 	TWeakPtr<FDocumentTabFactory> GraphEditorTabFactoryPtr;
 
-	/* The Script being edited */
-	UBehaviorTree* Script;
+	/* The Behavior Tree being edited */
+	UBehaviorTree* BehaviorTree;
+
+	/* The Blackboard Data being edited */
+	UBlackboardData* BlackboardData;
 
 	TWeakPtr<SGraphEditor> FocusedGraphEdPtr;
 	TWeakObjectPtr<class UBehaviorTreeGraphNode_CompositeDecorator> FocusedGraphOwner;
@@ -167,9 +202,10 @@ private:
 	uint32 bSelectedNodeIsInjected : 1;
 	uint32 SelectedNodesCount;
 
+	TSharedPtr<class FBehaviorTreeEditorToolbar> ToolbarBuilder;
+
 public:
-	/**	Graph editor tab */
-	static const FName BTGraphTabId;
-	static const FName BTPropertiesTabId;
-	static const FName BTSearchTabId;
+	/** Modes in mode switcher */
+	static const FName BehaviorTreeMode;
+	static const FName BlackboardMode;
 };
diff --git a/Engine/Source/Editor/BehaviorTreeEditor/Private/BehaviorTreeEditorCommands.cpp b/Engine/Source/Editor/BehaviorTreeEditor/Private/BehaviorTreeEditorCommands.cpp
new file mode 100644
--- /dev/null
+++ b/Engine/Source/Editor/BehaviorTreeEditor/Private/BehaviorTreeEditorCommands.cpp
@@ -0,0 +1,36 @@
+// Copyright 1998-2014 Epic Games, Inc. All Rights Reserved.
+
+#include "BehaviorTreeEditorPrivatePCH.h"
+#include "BehaviorTreeEditorCommands.h"
+
+#define LOCTEXT_NAMESPACE "BehaviorTreeEditorCommands"
+
+FBTCommonCommands::FBTCommonCommands() 
+	: TCommands<FBTCommonCommands>("BTEditor.Common", LOCTEXT("Common", "Common"), NAME_None, FEditorStyle::GetStyleSetName())
+{
+}
+
+void FBTCommonCommands::RegisterCommands()
+{
+	UI_COMMAND(SearchBT, "Search", "Search this Behavior Tree.", EUserInterfaceActionType::Button, FInputGesture(EModifierKey::Control, EKeys::F));
+}
+
+FBTDebuggerCommands::FBTDebuggerCommands() 
+	: TCommands<FBTDebuggerCommands>("BTEditor.Debugger", LOCTEXT("Debugger", "Debugger"), NAME_None, FEditorStyle::GetStyleSetName())
+{
+}
+
+void FBTDebuggerCommands::RegisterCommands()
+{
+	UI_COMMAND(BackInto, "Back: Into", "Show state from previous step, can go into subtrees", EUserInterfaceActionType::Button, FInputGesture());
+	UI_COMMAND(BackOver, "Back: Over", "Show state from previous step, don't go into subtrees", EUserInterfaceActionType::Button, FInputGesture());
+	UI_COMMAND(ForwardInto, "Forward: Into", "Show state from next step, can go into subtrees", EUserInterfaceActionType::Button, FInputGesture());
+	UI_COMMAND(ForwardOver, "Forward: Over", "Show state from next step, don't go into subtrees", EUserInterfaceActionType::Button, FInputGesture());
+	UI_COMMAND(StepOut, "Step Out", "Show state from next step, leave current subtree", EUserInterfaceActionType::Button, FInputGesture());
+	
+	UI_COMMAND(PausePlaySession, "Pause", "Pause simulation", EUserInterfaceActionType::Button, FInputGesture());
+	UI_COMMAND(ResumePlaySession, "Resume", "Resume simulation", EUserInterfaceActionType::Button, FInputGesture() );
+	UI_COMMAND(StopPlaySession, "Stop", "Stop simulation", EUserInterfaceActionType::Button, FInputGesture());
+}
+
+#undef LOCTEXT_NAMESPACE
diff --git a/Engine/Source/Editor/BehaviorTreeEditor/Private/BehaviorTreeEditorCommands.h b/Engine/Source/Editor/BehaviorTreeEditor/Private/BehaviorTreeEditorCommands.h
new file mode 100644
--- /dev/null
+++ b/Engine/Source/Editor/BehaviorTreeEditor/Private/BehaviorTreeEditorCommands.h
@@ -0,0 +1,35 @@
+// Copyright 1998-2014 Epic Games, Inc. All Rights Reserved.
+
+#pragma once
+
+#include "Commands.h"
+
+class FBTCommonCommands : public TCommands<FBTCommonCommands>
+{
+public:
+	FBTCommonCommands();
+
+	TSharedPtr<FUICommandInfo> SearchBT;
+
+	/** Initialize commands */
+	virtual void RegisterCommands() override;
+};
+
+class FBTDebuggerCommands : public TCommands<FBTDebuggerCommands>
+{
+public:
+	FBTDebuggerCommands();
+
+	TSharedPtr<FUICommandInfo> BackInto;
+	TSharedPtr<FUICommandInfo> BackOver;
+	TSharedPtr<FUICommandInfo> ForwardInto;
+	TSharedPtr<FUICommandInfo> ForwardOver;
+	TSharedPtr<FUICommandInfo> StepOut;
+
+	TSharedPtr<FUICommandInfo> PausePlaySession;
+	TSharedPtr<FUICommandInfo> ResumePlaySession;
+	TSharedPtr<FUICommandInfo> StopPlaySession;
+
+	/** Initialize commands */
+	virtual void RegisterCommands() override;
+};
diff --git a/Engine/Source/Editor/BehaviorTreeEditor/Private/BehaviorTreeEditorModes.cpp b/Engine/Source/Editor/BehaviorTreeEditor/Private/BehaviorTreeEditorModes.cpp
new file mode 100644
--- /dev/null
+++ b/Engine/Source/Editor/BehaviorTreeEditor/Private/BehaviorTreeEditorModes.cpp
@@ -0,0 +1,135 @@
+// Copyright 1998-2014 Epic Games, Inc. All Rights Reserved.
+
+#include "BehaviorTreeEditorPrivatePCH.h"
+#include "BehaviorTreeEditorModes.h"
+#include "BehaviorTreeEditorTabs.h"
+
+FBehaviorTreeEditorApplicationMode::FBehaviorTreeEditorApplicationMode(TSharedPtr<class FBehaviorTreeEditor> InBehaviorTreeEditor)
+	: FApplicationMode(FBehaviorTreeEditor::BehaviorTreeMode)
+{
+	BehaviorTreeEditor = InBehaviorTreeEditor;
+
+	BehaviorTreeEditorTabFactories.RegisterFactory(MakeShareable(new FBehaviorTreeDetailsSummoner(InBehaviorTreeEditor)));
+	BehaviorTreeEditorTabFactories.RegisterFactory(MakeShareable(new FBehaviorTreeSearchSummoner(InBehaviorTreeEditor)));
+
+	TabLayout = FTabManager::NewLayout( "Standalone_BehaviorTree_Layout" )
+	->AddArea
+	(
+		FTabManager::NewPrimaryArea() ->SetOrientation(Orient_Vertical)
+		->Split
+		(
+			FTabManager::NewStack()
+			->SetSizeCoefficient(0.1f)
+			->AddTab(InBehaviorTreeEditor->GetToolbarTabId(), ETabState::OpenedTab) 
+			->SetHideTabWell(true) 
+		)
+		->Split
+		(
+			FTabManager::NewSplitter() ->SetOrientation(Orient_Horizontal)
+			->Split
+			(
+				FTabManager::NewStack()
+				->SetSizeCoefficient(0.7f)
+				->AddTab(FBehaviorTreeEditorTabs::GraphEditorID, ETabState::ClosedTab)
+			)
+			->Split
+			(
+				FTabManager::NewStack()
+				->SetSizeCoefficient(0.3f)
+				->AddTab(FBehaviorTreeEditorTabs::GraphDetailsID, ETabState::OpenedTab)
+				->AddTab(FBehaviorTreeEditorTabs::SearchID, ETabState::ClosedTab)
+			)
+		)
+	);
+	
+	InBehaviorTreeEditor->GetToolbarBuilder()->AddModesToolbar(ToolbarExtender);
+	InBehaviorTreeEditor->GetToolbarBuilder()->AddDebuggerToolbar(ToolbarExtender);
+}
+
+void FBehaviorTreeEditorApplicationMode::RegisterTabFactories(TSharedPtr<FTabManager> InTabManager)
+{
+	check(BehaviorTreeEditor.IsValid());
+	TSharedPtr<FBehaviorTreeEditor> BehaviorTreeEditorPtr = BehaviorTreeEditor.Pin();
+	
+	BehaviorTreeEditorPtr->RegisterToolbarTab(InTabManager.ToSharedRef());
+
+	// Mode-specific setup
+	BehaviorTreeEditorPtr->PushTabFactories(BehaviorTreeEditorTabFactories);
+
+	FApplicationMode::RegisterTabFactories(InTabManager);
+}
+
+void FBehaviorTreeEditorApplicationMode::PreDeactivateMode()
+{
+	FApplicationMode::PreDeactivateMode();
+
+	check(BehaviorTreeEditor.IsValid());
+	TSharedPtr<FBehaviorTreeEditor> BehaviorTreeEditorPtr = BehaviorTreeEditor.Pin();
+	
+	BehaviorTreeEditorPtr->SaveEditedObjectState();
+}
+
+void FBehaviorTreeEditorApplicationMode::PostActivateMode()
+{
+	// Reopen any documents that were open when the blueprint was last saved
+	check(BehaviorTreeEditor.IsValid());
+	TSharedPtr<FBehaviorTreeEditor> BehaviorTreeEditorPtr = BehaviorTreeEditor.Pin();
+	BehaviorTreeEditorPtr->RestoreBehaviorTree();
+
+	FApplicationMode::PostActivateMode();
+}
+
+FBlackboardEditorApplicationMode::FBlackboardEditorApplicationMode(TSharedPtr<class FBehaviorTreeEditor> InBehaviorTreeEditor)
+	: FApplicationMode(FBehaviorTreeEditor::BlackboardMode)
+{
+	BehaviorTreeEditor = InBehaviorTreeEditor;
+	
+	BlackboardTabFactories.RegisterFactory(MakeShareable(new FBlackboardDetailsSummoner(InBehaviorTreeEditor)));
+
+	TabLayout = FTabManager::NewLayout( "Standalone_BlackboardEditor_Layout" )
+	->AddArea
+	(
+		FTabManager::NewPrimaryArea() ->SetOrientation(Orient_Vertical)
+		->Split
+		(
+			FTabManager::NewStack()
+			->SetSizeCoefficient(0.1f)
+			->SetHideTabWell( true )
+			->AddTab(InBehaviorTreeEditor->GetToolbarTabId(), ETabState::OpenedTab)
+		)
+		->Split
+		(
+			FTabManager::NewSplitter()
+			->Split
+			(
+				FTabManager::NewStack()
+				->AddTab(FBehaviorTreeEditorTabs::BlackboardDetailsID, ETabState::OpenedTab)
+			)
+		)
+	);
+
+	InBehaviorTreeEditor->GetToolbarBuilder()->AddModesToolbar(ToolbarExtender);
+}
+
+void FBlackboardEditorApplicationMode::RegisterTabFactories(TSharedPtr<FTabManager> InTabManager)
+{
+	check(BehaviorTreeEditor.IsValid());
+	TSharedPtr<FBehaviorTreeEditor> BehaviorTreeEditorPtr = BehaviorTreeEditor.Pin();
+	
+	BehaviorTreeEditorPtr->RegisterToolbarTab(InTabManager.ToSharedRef());
+
+	// Mode-specific setup
+	BehaviorTreeEditorPtr->PushTabFactories(BlackboardTabFactories);
+
+	FApplicationMode::RegisterTabFactories(InTabManager);
+}
+
+void FBlackboardEditorApplicationMode::PostActivateMode()
+{
+	// Reopen any documents that were open when the BT was last saved
+	check(BehaviorTreeEditor.IsValid());
+	TSharedPtr<FBehaviorTreeEditor> BehaviorTreeEditorPtr = BehaviorTreeEditor.Pin();
+//	BehaviorTreeEditorPtr->StartEditingDefaults();
+
+	FApplicationMode::PostActivateMode();
+}
diff --git a/Engine/Source/Editor/BehaviorTreeEditor/Private/BehaviorTreeEditorModes.h b/Engine/Source/Editor/BehaviorTreeEditor/Private/BehaviorTreeEditorModes.h
new file mode 100644
--- /dev/null
+++ b/Engine/Source/Editor/BehaviorTreeEditor/Private/BehaviorTreeEditorModes.h
@@ -0,0 +1,38 @@
+// Copyright 1998-2014 Epic Games, Inc. All Rights Reserved.
+
+#pragma once
+
+#include "WorkflowOrientedApp/ApplicationMode.h"
+
+/** Application mode for main behavior tree editing mode */
+class FBehaviorTreeEditorApplicationMode : public FApplicationMode
+{
+public:
+	FBehaviorTreeEditorApplicationMode(TSharedPtr<class FBehaviorTreeEditor> InBehaviorTreeEditor);
+
+	virtual void RegisterTabFactories(TSharedPtr<class FTabManager> InTabManager) OVERRIDE;
+	virtual void PreDeactivateMode() OVERRIDE;
+	virtual void PostActivateMode() OVERRIDE;
+
+protected:
+	TWeakPtr<class FBehaviorTreeEditor> BehaviorTreeEditor;
+
+	// Set of spawnable tabs in behavior tree editing mode
+	FWorkflowAllowedTabSet BehaviorTreeEditorTabFactories;
+};
+
+/** Application mode for blackboard editing mode */
+class FBlackboardEditorApplicationMode : public FApplicationMode
+{
+public:
+	FBlackboardEditorApplicationMode(TSharedPtr<class FBehaviorTreeEditor> InBehaviorTreeEditor);
+
+	virtual void RegisterTabFactories(TSharedPtr<class FTabManager> InTabManager) OVERRIDE;
+	virtual void PostActivateMode() OVERRIDE;
+
+protected:
+	TWeakPtr<class FBehaviorTreeEditor> BehaviorTreeEditor;
+
+	// Set of spawnable tabs in blackboard mode
+	FWorkflowAllowedTabSet BlackboardTabFactories;
+};
\ No newline at end of file
diff --git a/Engine/Source/Editor/BehaviorTreeEditor/Private/BehaviorTreeEditorModule.cpp b/Engine/Source/Editor/BehaviorTreeEditor/Private/BehaviorTreeEditorModule.cpp
--- a/Engine/Source/Editor/BehaviorTreeEditor/Private/BehaviorTreeEditorModule.cpp
+++ b/Engine/Source/Editor/BehaviorTreeEditor/Private/BehaviorTreeEditorModule.cpp
@@ -45,8 +45,14 @@ void FBehaviorTreeEditorModule::StartupModule()
 	GraphPanelNodeFactory_BehaviorTree = MakeShareable( new FGraphPanelNodeFactory_BehaviorTree() );
 	FEdGraphUtilities::RegisterVisualNodeFactory(GraphPanelNodeFactory_BehaviorTree);
 
-	ItemDataAssetTypeActions = MakeShareable(new FAssetTypeActions_BehaviorTree);
-	FModuleManager::LoadModuleChecked<FAssetToolsModule>("AssetTools").Get().RegisterAssetTypeActions(ItemDataAssetTypeActions.ToSharedRef());
+	IAssetTools& AssetToolsModule = FModuleManager::LoadModuleChecked<FAssetToolsModule>("AssetTools").Get();
+	TSharedPtr<FAssetTypeActions_BehaviorTree> BehaviorTreeAssetTypeAction = MakeShareable(new FAssetTypeActions_BehaviorTree);
+	ItemDataAssetTypeActions.Add(BehaviorTreeAssetTypeAction);
+	AssetToolsModule.RegisterAssetTypeActions(BehaviorTreeAssetTypeAction.ToSharedRef());
+
+	TSharedPtr<FAssetTypeActions_Blackboard> BlackboardAssetTypeAction = MakeShareable(new FAssetTypeActions_Blackboard);
+	ItemDataAssetTypeActions.Add(BlackboardAssetTypeAction);
+	AssetToolsModule.RegisterAssetTypeActions(BlackboardAssetTypeAction.ToSharedRef());
 
 	// Register the details customizer
 	FPropertyEditorModule& PropertyModule = FModuleManager::LoadModuleChecked<FPropertyEditorModule>("PropertyEditor");
@@ -75,14 +81,18 @@ void FBehaviorTreeEditorModule::ShutdownModule()
 	}
 
 	// Unregister the BehaviorTree item data asset type actions
-	if (ItemDataAssetTypeActions.IsValid())
+	if (FModuleManager::Get().IsModuleLoaded("AssetTools"))
 	{
-		if (FModuleManager::Get().IsModuleLoaded("AssetTools"))
+		IAssetTools& AssetToolsModule = FModuleManager::GetModuleChecked<FAssetToolsModule>("AssetTools").Get();
+		for(auto& AssetTypeAction : ItemDataAssetTypeActions)
 		{
-			FModuleManager::GetModuleChecked<FAssetToolsModule>("AssetTools").Get().UnregisterAssetTypeActions(ItemDataAssetTypeActions.ToSharedRef());
-		}
-		ItemDataAssetTypeActions.Reset();
+			if (AssetTypeAction.IsValid())
+			{
+				AssetToolsModule.UnregisterAssetTypeActions(AssetTypeAction.ToSharedRef());
+			}	
+		}			
 	}
+	ItemDataAssetTypeActions.Empty();
 
 	// Unregister the details customization
 	if (FModuleManager::Get().IsModuleLoaded("PropertyEditor"))
@@ -104,7 +114,7 @@ void FBehaviorTreeEditorModule::HandleExperimentalSettingChanged(FName PropertyN
 	}
 }
 
-TSharedRef<IBehaviorTreeEditor> FBehaviorTreeEditorModule::CreateBehaviorTreeEditor( const EToolkitMode::Type Mode, const TSharedPtr< IToolkitHost >& InitToolkitHost, UBehaviorTree* Script )
+TSharedRef<IBehaviorTreeEditor> FBehaviorTreeEditorModule::CreateBehaviorTreeEditor( const EToolkitMode::Type Mode, const TSharedPtr< IToolkitHost >& InitToolkitHost, UObject* Object )
 {
 	if (!ClassCache.IsValid())
 	{
@@ -112,7 +122,6 @@ TSharedRef<IBehaviorTreeEditor> FBehaviorTreeEditorModule::CreateBehaviorTreeEdi
 	}
 
 	TSharedRef< FBehaviorTreeEditor > NewBehaviorTreeEditor( new FBehaviorTreeEditor() );
-	NewBehaviorTreeEditor->InitBehaviorTreeEditor( Mode, InitToolkitHost, Script );
-	return NewBehaviorTreeEditor;
+	NewBehaviorTreeEditor->InitBehaviorTreeEditor( Mode, InitToolkitHost, Object );
+	return NewBehaviorTreeEditor;	
 }
-
diff --git a/Engine/Source/Editor/BehaviorTreeEditor/Private/BehaviorTreeEditorPrivatePCH.h b/Engine/Source/Editor/BehaviorTreeEditor/Private/BehaviorTreeEditorPrivatePCH.h
--- a/Engine/Source/Editor/BehaviorTreeEditor/Private/BehaviorTreeEditorPrivatePCH.h
+++ b/Engine/Source/Editor/BehaviorTreeEditor/Private/BehaviorTreeEditorPrivatePCH.h
@@ -16,6 +16,7 @@
 
 #include "BehaviorTreeEditorClasses.h"
 #include "AssetTypeActions_BehaviorTree.h"
+#include "AssetTypeActions_Blackboard.h"
 
 #include "BehaviorTreeEditorModule.h"
 #include "BehaviorTreeEditor.h"
diff --git a/Engine/Source/Editor/BehaviorTreeEditor/Private/BehaviorTreeEditorTabFactories.cpp b/Engine/Source/Editor/BehaviorTreeEditor/Private/BehaviorTreeEditorTabFactories.cpp
new file mode 100644
--- /dev/null
+++ b/Engine/Source/Editor/BehaviorTreeEditor/Private/BehaviorTreeEditorTabFactories.cpp
@@ -0,0 +1,143 @@
+// Copyright 1998-2014 Epic Games, Inc. All Rights Reserved.
+
+#include "BehaviorTreeEditorPrivatePCH.h"
+#include "BehaviorTreeEditorTabFactories.h"
+#include "BehaviorTreeEditor.h"
+#include "BehaviorTreeEditorTabs.h"
+
+#define LOCTEXT_NAMESPACE "BehaviorTreeEditorFactories"
+
+FBlackboardDetailsSummoner::FBlackboardDetailsSummoner(TSharedPtr<class FBehaviorTreeEditor> InBehaviorTreeEditorPtr)
+	: FWorkflowTabFactory(FBehaviorTreeEditorTabs::BlackboardDetailsID, InBehaviorTreeEditorPtr)
+	, BehaviorTreeEditorPtr(InBehaviorTreeEditorPtr)
+{
+	TabLabel = LOCTEXT("BlackboardDetailsLabel", "Blackboard Details");
+	TabIcon = FEditorStyle::GetBrush("Kismet.Tabs.Components");
+
+	bIsSingleton = true;
+
+	ViewMenuDescription = LOCTEXT("BlackboardDetailsView", "Blackboard Details");
+	ViewMenuTooltip = LOCTEXT("BlackboardDetailsView_ToolTip", "Show the blackboard details view");
+}
+
+TSharedRef<SWidget> FBlackboardDetailsSummoner::CreateTabBody(const FWorkflowTabSpawnInfo& Info) const
+{
+	const bool bIsUpdatable = false;
+	const bool bAllowFavorites = true;
+	const bool bIsLockable = false;
+	const bool bAllowSearch = true;
+	const bool bObjectsUseNameArea = false;
+	const bool bHideSelectionTip = true;
+
+	FPropertyEditorModule& PropertyEditorModule = FModuleManager::GetModuleChecked<FPropertyEditorModule>( "PropertyEditor" );
+	const FDetailsViewArgs DetailsViewArgs( bIsUpdatable, bIsLockable, bAllowSearch, bObjectsUseNameArea, bHideSelectionTip );
+	TSharedRef<IDetailsView> DetailsView = PropertyEditorModule.CreateDetailView( DetailsViewArgs );
+
+	DetailsView->SetObject(BehaviorTreeEditorPtr.Pin()->GetBlackboardData());
+
+	return DetailsView;
+}
+
+FText FBlackboardDetailsSummoner::GetTabToolTipText(const FWorkflowTabSpawnInfo& Info) const
+{
+	return LOCTEXT("BlackboardTabTooltip", "The Blackboard tab is for editing blackboard key/value pairs.");
+}
+
+FBehaviorTreeDetailsSummoner::FBehaviorTreeDetailsSummoner(TSharedPtr<class FBehaviorTreeEditor> InBehaviorTreeEditorPtr)
+	: FWorkflowTabFactory(FBehaviorTreeEditorTabs::GraphDetailsID, InBehaviorTreeEditorPtr)
+	, BehaviorTreeEditorPtr(InBehaviorTreeEditorPtr)
+{
+	TabLabel = LOCTEXT("BehaviorTreeDetailsLabel", "Details");
+	TabIcon = FEditorStyle::GetBrush("Kismet.Tabs.Components");
+
+	bIsSingleton = true;
+
+	ViewMenuDescription = LOCTEXT("BehaviorTreeDetailsView", "Behavior Tree Details");
+	ViewMenuTooltip = LOCTEXT("BehaviorTreeDetailsView_ToolTip", "Show the behavior tree details view");
+}
+
+TSharedRef<SWidget> FBehaviorTreeDetailsSummoner::CreateTabBody(const FWorkflowTabSpawnInfo& Info) const
+{
+	check(BehaviorTreeEditorPtr.IsValid());
+	return BehaviorTreeEditorPtr.Pin()->SpawnProperties();
+}
+
+FText FBehaviorTreeDetailsSummoner::GetTabToolTipText(const FWorkflowTabSpawnInfo& Info) const
+{
+	return LOCTEXT("BehaviorTreeDetailsTabTooltip", "The behavior tree details tab allows editing of the properties of behavior tree nodes");
+}
+
+FBehaviorTreeSearchSummoner::FBehaviorTreeSearchSummoner(TSharedPtr<class FBehaviorTreeEditor> InBehaviorTreeEditorPtr)
+	: FWorkflowTabFactory(FBehaviorTreeEditorTabs::SearchID, InBehaviorTreeEditorPtr)
+	, BehaviorTreeEditorPtr(InBehaviorTreeEditorPtr)
+{
+	TabLabel = LOCTEXT("BehaviorTreeSearchLabel", "Search");
+	TabIcon = FEditorStyle::GetBrush("Kismet.Tabs.Components");
+
+	bIsSingleton = true;
+
+	ViewMenuDescription = LOCTEXT("BehaviorTreeSearchView", "Search");
+	ViewMenuTooltip = LOCTEXT("BehaviorTreeSearchView_ToolTip", "Show the behavior tree search tab");
+}
+
+TSharedRef<SWidget> FBehaviorTreeSearchSummoner::CreateTabBody(const FWorkflowTabSpawnInfo& Info) const
+{
+	return BehaviorTreeEditorPtr.Pin()->SpawnSearch();
+}
+
+FText FBehaviorTreeSearchSummoner::GetTabToolTipText(const FWorkflowTabSpawnInfo& Info) const
+{
+	return LOCTEXT("BehaviorTreeSearchTabTooltip", "The behavior tree search tab allows searching within behavior tree nodes");
+}
+
+FBTGraphEditorSummoner::FBTGraphEditorSummoner(TSharedPtr<class FBehaviorTreeEditor> InBehaviorTreeEditorPtr, FOnCreateGraphEditorWidget CreateGraphEditorWidgetCallback)
+	: FDocumentTabFactoryForObjects<UEdGraph>(FBehaviorTreeEditorTabs::GraphEditorID, InBehaviorTreeEditorPtr)
+	, BehaviorTreeEditorPtr(InBehaviorTreeEditorPtr)
+	, OnCreateGraphEditorWidget(CreateGraphEditorWidgetCallback)
+{
+}
+
+void FBTGraphEditorSummoner::OnTabActivated(TSharedPtr<SDockTab> Tab) const
+{
+	check(BehaviorTreeEditorPtr.IsValid());
+	TSharedRef<SGraphEditor> GraphEditor = StaticCastSharedRef<SGraphEditor>(Tab->GetContent());
+	BehaviorTreeEditorPtr.Pin()->OnGraphEditorFocused(GraphEditor);
+}
+
+void FBTGraphEditorSummoner::OnTabRefreshed(TSharedPtr<SDockTab> Tab) const
+{
+	TSharedRef<SGraphEditor> GraphEditor = StaticCastSharedRef<SGraphEditor>(Tab->GetContent());
+	GraphEditor->NotifyGraphChanged();
+}
+
+TAttribute<FText> FBTGraphEditorSummoner::ConstructTabNameForObject(UEdGraph* DocumentID) const
+{
+	return TAttribute<FText>( FText::FromString( DocumentID->GetName() ) );
+}
+
+TSharedRef<SWidget> FBTGraphEditorSummoner::CreateTabBodyForObject(const FWorkflowTabSpawnInfo& Info, UEdGraph* DocumentID) const
+{
+	return OnCreateGraphEditorWidget.Execute(DocumentID);
+}
+
+const FSlateBrush* FBTGraphEditorSummoner::GetTabIconForObject(const FWorkflowTabSpawnInfo& Info, UEdGraph* DocumentID) const
+{
+	return FEditorStyle::GetBrush("NoBrush");
+}
+
+void FBTGraphEditorSummoner::SaveState(TSharedPtr<SDockTab> Tab, TSharedPtr<FTabPayload> Payload) const
+{
+	check(BehaviorTreeEditorPtr.IsValid());
+	check(BehaviorTreeEditorPtr.Pin()->GetBehaviorTree());
+
+	TSharedRef<SGraphEditor> GraphEditor = StaticCastSharedRef<SGraphEditor>(Tab->GetContent());
+
+	FVector2D ViewLocation;
+	float ZoomAmount;
+	GraphEditor->GetViewLocation(ViewLocation, ZoomAmount);
+
+	UEdGraph* Graph = FTabPayload_UObject::CastChecked<UEdGraph>(Payload);
+	BehaviorTreeEditorPtr.Pin()->GetBehaviorTree()->LastEditedDocuments.Add(FEditedDocumentInfo(Graph, ViewLocation, ZoomAmount));
+}
+
+#undef LOCTEXT_NAMESPACE
diff --git a/Engine/Source/Editor/BehaviorTreeEditor/Private/BehaviorTreeEditorTabFactories.h b/Engine/Source/Editor/BehaviorTreeEditor/Private/BehaviorTreeEditorTabFactories.h
new file mode 100644
--- /dev/null
+++ b/Engine/Source/Editor/BehaviorTreeEditor/Private/BehaviorTreeEditorTabFactories.h
@@ -0,0 +1,63 @@
+// Copyright 1998-2014 Epic Games, Inc. All Rights Reserved.
+
+#pragma once
+
+#include "WorkflowOrientedApp/WorkflowTabFactory.h"
+#include "WorkflowOrientedApp/WorkflowUObjectDocuments.h"
+
+struct FBlackboardDetailsSummoner : public FWorkflowTabFactory
+{
+public:
+	FBlackboardDetailsSummoner(TSharedPtr<class FBehaviorTreeEditor> InBehaviorTreeEditorPtr);
+
+	virtual TSharedRef<SWidget> CreateTabBody(const FWorkflowTabSpawnInfo& Info) const override;
+	virtual FText GetTabToolTipText(const FWorkflowTabSpawnInfo& Info) const override;
+
+protected:
+	TWeakPtr<class FBehaviorTreeEditor> BehaviorTreeEditorPtr;
+};
+
+struct FBehaviorTreeDetailsSummoner : public FWorkflowTabFactory
+{
+public:
+	FBehaviorTreeDetailsSummoner(TSharedPtr<class FBehaviorTreeEditor> InBehaviorTreeEditorPtr);
+
+	virtual TSharedRef<SWidget> CreateTabBody(const FWorkflowTabSpawnInfo& Info) const override;
+	virtual FText GetTabToolTipText(const FWorkflowTabSpawnInfo& Info) const override;
+
+protected:
+	TWeakPtr<class FBehaviorTreeEditor> BehaviorTreeEditorPtr;
+};
+
+struct FBehaviorTreeSearchSummoner : public FWorkflowTabFactory
+{
+public:
+	FBehaviorTreeSearchSummoner(TSharedPtr<class FBehaviorTreeEditor> InBehaviorTreeEditorPtr);
+
+	virtual TSharedRef<SWidget> CreateTabBody(const FWorkflowTabSpawnInfo& Info) const override;
+	virtual FText GetTabToolTipText(const FWorkflowTabSpawnInfo& Info) const override;
+
+protected:
+	TWeakPtr<class FBehaviorTreeEditor> BehaviorTreeEditorPtr;
+};
+
+struct FBTGraphEditorSummoner : public FDocumentTabFactoryForObjects<UEdGraph>
+{
+public:
+	DECLARE_DELEGATE_RetVal_OneParam(TSharedRef<SGraphEditor>, FOnCreateGraphEditorWidget, UEdGraph*);
+public:
+	FBTGraphEditorSummoner(TSharedPtr<class FBehaviorTreeEditor> InBehaviorTreeEditorPtr, FOnCreateGraphEditorWidget CreateGraphEditorWidgetCallback);
+
+	virtual void OnTabActivated(TSharedPtr<SDockTab> Tab) const override;
+	virtual void OnTabRefreshed(TSharedPtr<SDockTab> Tab) const override;
+
+protected:
+	virtual TAttribute<FText> ConstructTabNameForObject(UEdGraph* DocumentID) const override;
+	virtual TSharedRef<SWidget> CreateTabBodyForObject(const FWorkflowTabSpawnInfo& Info, UEdGraph* DocumentID) const override;
+	virtual const FSlateBrush* GetTabIconForObject(const FWorkflowTabSpawnInfo& Info, UEdGraph* DocumentID) const override;
+	virtual void SaveState(TSharedPtr<SDockTab> Tab, TSharedPtr<FTabPayload> Payload) const override;
+
+protected:
+	TWeakPtr<class FBehaviorTreeEditor> BehaviorTreeEditorPtr;
+	FOnCreateGraphEditorWidget OnCreateGraphEditorWidget;
+};
diff --git a/Engine/Source/Editor/BehaviorTreeEditor/Private/BehaviorTreeEditorTabs.cpp b/Engine/Source/Editor/BehaviorTreeEditor/Private/BehaviorTreeEditorTabs.cpp
new file mode 100644
--- /dev/null
+++ b/Engine/Source/Editor/BehaviorTreeEditor/Private/BehaviorTreeEditorTabs.cpp
@@ -0,0 +1,12 @@
+// Copyright 1998-2014 Epic Games, Inc. All Rights Reserved.
+
+#include "BehaviorTreeEditorPrivatePCH.h"
+#include "BehaviorTreeEditorTabs.h"
+
+// Tab identifiers
+const FName FBehaviorTreeEditorTabs::GraphDetailsID(TEXT("BehaviorTreeEditor_Properties"));
+const FName FBehaviorTreeEditorTabs::SearchID(TEXT("BehaviorTreeEditor_Search"));
+const FName FBehaviorTreeEditorTabs::BlackboardDetailsID(TEXT("BehaviorTreeEditor_BlackboardProperties"));
+
+// Document tab identifiers
+const FName FBehaviorTreeEditorTabs::GraphEditorID(TEXT("Document"));
diff --git a/Engine/Source/Editor/BehaviorTreeEditor/Private/BehaviorTreeEditorTabs.h b/Engine/Source/Editor/BehaviorTreeEditor/Private/BehaviorTreeEditorTabs.h
new file mode 100644
--- /dev/null
+++ b/Engine/Source/Editor/BehaviorTreeEditor/Private/BehaviorTreeEditorTabs.h
@@ -0,0 +1,15 @@
+// Copyright 1998-2014 Epic Games, Inc. All Rights Reserved.
+
+#pragma once
+
+struct FBehaviorTreeEditorTabs
+{
+	// Tab identifiers
+	static const FName GraphDetailsID;
+	static const FName SearchID;
+	static const FName BlackboardDetailsID;
+
+	// Document tab identifiers
+	static const FName GraphEditorID;
+};
+
diff --git a/Engine/Source/Editor/BehaviorTreeEditor/Private/BehaviorTreeEditorToolbar.cpp b/Engine/Source/Editor/BehaviorTreeEditor/Private/BehaviorTreeEditorToolbar.cpp
new file mode 100644
--- /dev/null
+++ b/Engine/Source/Editor/BehaviorTreeEditor/Private/BehaviorTreeEditorToolbar.cpp
@@ -0,0 +1,136 @@
+// Copyright 1998-2014 Epic Games, Inc. All Rights Reserved.
+
+#include "BehaviorTreeEditorPrivatePCH.h"
+#include "BehaviorTreeEditorToolbar.h"
+#include "BehaviorTreeEditor.h"
+#include "WorkflowOrientedApp/SModeWidget.h"
+#include "BehaviorTreeEditorCommands.h"
+
+#define LOCTEXT_NAMESPACE "BehaviorTreeEditorToolbar"
+
+class SBehaviorTreeModeSeparator : public SBorder
+{
+public:
+	SLATE_BEGIN_ARGS(SBehaviorTreeModeSeparator) {}
+	SLATE_END_ARGS()
+
+	void Construct(const FArguments& InArg)
+	{
+		SBorder::Construct(
+			SBorder::FArguments()
+			.BorderImage(FEditorStyle::GetBrush("BlueprintEditor.PipelineSeparator"))
+			.Padding(0.0f)
+			);
+	}
+
+	// SWidget interface
+	virtual FVector2D ComputeDesiredSize() const OVERRIDE
+	{
+		const float Height = 20.0f;
+		const float Thickness = 16.0f;
+		return FVector2D(Thickness, Height);
+	}
+	// End of SWidget interface
+};
+
+void FBehaviorTreeEditorToolbar::AddModesToolbar(TSharedPtr<FExtender> Extender)
+{
+	check(BehaviorTreeEditor.IsValid());
+	TSharedPtr<FBehaviorTreeEditor> BehaviorTreeEditorPtr = BehaviorTreeEditor.Pin();
+
+	Extender->AddToolBarExtension(
+		"Asset",
+		EExtensionHook::After,
+		BehaviorTreeEditorPtr->GetToolkitCommands(),
+		FToolBarExtensionDelegate::CreateSP( this, &FBehaviorTreeEditorToolbar::FillModesToolbar ) );
+}
+
+void FBehaviorTreeEditorToolbar::AddDebuggerToolbar(TSharedPtr<FExtender> Extender)
+{
+	check(BehaviorTreeEditor.IsValid());
+	TSharedPtr<FBehaviorTreeEditor> BehaviorTreeEditorPtr = BehaviorTreeEditor.Pin();
+
+	// setup toolbar
+	struct Local
+	{
+		static void FillToolbar(FToolBarBuilder& ToolbarBuilder, TWeakPtr<FBehaviorTreeEditor> BehaviorTreeEditor)
+		{
+			TSharedPtr<FBehaviorTreeEditor> BehaviorTreeEditorPtr = BehaviorTreeEditor.Pin();
+
+			const bool bCanShowDebugger = BehaviorTreeEditorPtr->IsDebuggerReady();
+			if (bCanShowDebugger)
+			{
+				TSharedRef<SWidget> SelectionBox = SNew(SComboButton)
+					.OnGetMenuContent( BehaviorTreeEditorPtr.Get(), &FBehaviorTreeEditor::OnGetDebuggerActorsMenu )
+					.ButtonContent()
+					[
+						SNew(STextBlock)
+						.ToolTipText( LOCTEXT("SelectDebugActor", "Pick actor to debug") )
+						.Text(BehaviorTreeEditorPtr.Get(), &FBehaviorTreeEditor::GetDebuggerActorDesc )
+					];
+
+				ToolbarBuilder.BeginSection("CachedState");
+				{
+					ToolbarBuilder.AddToolBarButton(FBTDebuggerCommands::Get().BackOver);
+					ToolbarBuilder.AddToolBarButton(FBTDebuggerCommands::Get().BackInto);
+					ToolbarBuilder.AddToolBarButton(FBTDebuggerCommands::Get().ForwardInto);
+					ToolbarBuilder.AddToolBarButton(FBTDebuggerCommands::Get().ForwardOver);
+					ToolbarBuilder.AddToolBarButton(FBTDebuggerCommands::Get().StepOut);
+				}
+				ToolbarBuilder.EndSection();
+				ToolbarBuilder.BeginSection("World");
+				{
+					ToolbarBuilder.AddToolBarButton(FBTDebuggerCommands::Get().PausePlaySession);
+					ToolbarBuilder.AddToolBarButton(FBTDebuggerCommands::Get().ResumePlaySession);
+					ToolbarBuilder.AddToolBarButton(FBTDebuggerCommands::Get().StopPlaySession);
+					ToolbarBuilder.AddSeparator();
+					ToolbarBuilder.AddWidget(SelectionBox);
+				}
+				ToolbarBuilder.EndSection();
+			}
+		}
+	};
+
+	TSharedPtr<FExtender> ToolbarExtender = MakeShareable(new FExtender);
+	ToolbarExtender->AddToolBarExtension("Asset", EExtensionHook::After, BehaviorTreeEditorPtr->GetToolkitCommands(), FToolBarExtensionDelegate::CreateStatic( &Local::FillToolbar, BehaviorTreeEditor ));
+	BehaviorTreeEditorPtr->AddToolbarExtender(ToolbarExtender);
+}
+
+void FBehaviorTreeEditorToolbar::FillModesToolbar(FToolBarBuilder& ToolbarBuilder)
+{
+	check(BehaviorTreeEditor.IsValid());
+	TSharedPtr<FBehaviorTreeEditor> BehaviorTreeEditorPtr = BehaviorTreeEditor.Pin();
+
+	TAttribute<FName> GetActiveMode(BehaviorTreeEditorPtr.ToSharedRef(), &FBehaviorTreeEditor::GetCurrentMode);
+	FOnModeChangeRequested SetActiveMode = FOnModeChangeRequested::CreateSP(BehaviorTreeEditorPtr.ToSharedRef(), &FBehaviorTreeEditor::SetCurrentMode);
+
+	// Left side padding
+	BehaviorTreeEditorPtr->AddToolbarWidget(SNew(SSpacer).Size(FVector2D(4.0f, 1.0f)));
+
+	BehaviorTreeEditorPtr->AddToolbarWidget(
+		SNew(SModeWidget, FBehaviorTreeEditor::GetLocalizedMode( FBehaviorTreeEditor::BehaviorTreeMode ), FBehaviorTreeEditor::BehaviorTreeMode)
+		.OnGetActiveMode(GetActiveMode)
+		.OnSetActiveMode(SetActiveMode)
+		.CanBeSelected(BehaviorTreeEditorPtr.Get(), &FBehaviorTreeEditor::CanAccessBehaviorTreeMode)
+		.ToolTipText(LOCTEXT("BehaviorTreeModeButtonTooltip", "Switch to Behavior Tree Mode"))
+		.IconImage(FEditorStyle::GetBrush("BTEditor.SwitchToBehaviorTreeMode"))
+		.SmallIconImage(FEditorStyle::GetBrush("BTEditor.SwitchToBehaviorTreeMode.Small"))
+	);
+
+	BehaviorTreeEditorPtr->AddToolbarWidget(SNew(SBehaviorTreeModeSeparator));
+
+	BehaviorTreeEditorPtr->AddToolbarWidget(
+		SNew(SModeWidget, FBehaviorTreeEditor::GetLocalizedMode( FBehaviorTreeEditor::BlackboardMode ), FBehaviorTreeEditor::BlackboardMode)
+		.OnGetActiveMode(GetActiveMode)
+		.OnSetActiveMode(SetActiveMode)
+		.CanBeSelected(BehaviorTreeEditorPtr.Get(), &FBehaviorTreeEditor::CanAccessBlackboardMode)
+		.ToolTipText(LOCTEXT("BlackboardModeButtonTooltip", "Switch to Blackboard Mode"))
+		.IconImage(FEditorStyle::GetBrush("BTEditor.SwitchToBlackboardMode"))
+		.SmallIconImage(FEditorStyle::GetBrush("BTEditor.SwitchToBlackboardMode.Small"))
+	);
+		
+	// Right side padding
+	BehaviorTreeEditorPtr->AddToolbarWidget(SNew(SSpacer).Size(FVector2D(4.0f, 1.0f)));
+}
+
+#undef LOCTEXT_NAMESPACE
\ No newline at end of file
diff --git a/Engine/Source/Editor/BehaviorTreeEditor/Private/BehaviorTreeEditorToolbar.h b/Engine/Source/Editor/BehaviorTreeEditor/Private/BehaviorTreeEditorToolbar.h
new file mode 100644
--- /dev/null
+++ b/Engine/Source/Editor/BehaviorTreeEditor/Private/BehaviorTreeEditorToolbar.h
@@ -0,0 +1,21 @@
+// Copyright 1998-2014 Epic Games, Inc. All Rights Reserved.
+
+#pragma once
+
+class FBehaviorTreeEditorToolbar : public TSharedFromThis<FBehaviorTreeEditorToolbar>
+{
+public:
+	FBehaviorTreeEditorToolbar(TSharedPtr<FBehaviorTreeEditor> InBehaviorTreeEditor)
+		: BehaviorTreeEditor(InBehaviorTreeEditor) {}
+
+	void AddModesToolbar(TSharedPtr<FExtender> Extender);
+	void AddDebuggerToolbar(TSharedPtr<FExtender> Extender);
+
+private:
+	void FillModesToolbar(FToolBarBuilder& ToolbarBuilder);
+	void FillDebuggerToolbar(FToolBarBuilder& ToolbarBuilder);
+
+protected:
+	/** Pointer back to the blueprint editor tool that owns us */
+	TWeakPtr<FBehaviorTreeEditor> BehaviorTreeEditor;
+};
diff --git a/Engine/Source/Editor/BehaviorTreeEditor/Public/BehaviorTreeEditorModule.h b/Engine/Source/Editor/BehaviorTreeEditor/Public/BehaviorTreeEditorModule.h
--- a/Engine/Source/Editor/BehaviorTreeEditor/Public/BehaviorTreeEditorModule.h
+++ b/Engine/Source/Editor/BehaviorTreeEditor/Public/BehaviorTreeEditorModule.h
@@ -19,11 +19,8 @@ public:
 	virtual void StartupModule() OVERRIDE;
 	virtual void ShutdownModule() OVERRIDE;
 
-	/** Compile the specified script. */
-	//virtual void CompileScript(class UNiagaraScript* ScriptToCompile);
-
-	/** Creates an instance of Niagara editor.  Only virtual so that it can be called across the DLL boundary. */
-	virtual TSharedRef<IBehaviorTreeEditor> CreateBehaviorTreeEditor( const EToolkitMode::Type Mode, const TSharedPtr< class IToolkitHost >& InitToolkitHost, class UBehaviorTree* Script );
+	/** Creates an instance of Behavior Tree editor.  Only virtual so that it can be called across the DLL boundary. */
+	virtual TSharedRef<IBehaviorTreeEditor> CreateBehaviorTreeEditor( const EToolkitMode::Type Mode, const TSharedPtr< class IToolkitHost >& InitToolkitHost, class UObject* Object );
 
 	/** Gets the extensibility managers for outside entities to extend static mesh editor's menus and toolbars */
 	virtual TSharedPtr<FExtensibilityManager> GetMenuExtensibilityManager() {return MenuExtensibilityManager;}
@@ -31,7 +28,7 @@ public:
 
 	TSharedPtr<struct FClassBrowseHelper> GetClassCache() { return ClassCache; }
 
-	/** Niagara Editor app identifier string */
+	/** Behavior Tree app identifier string */
 	static const FName BehaviorTreeEditorAppIdentifier;
 
 private:
@@ -41,7 +38,7 @@ private:
 	TSharedPtr<FExtensibilityManager> ToolBarExtensibilityManager;
 
 	/** Asset type actions */
-	TSharedPtr<class FAssetTypeActions_BehaviorTree> ItemDataAssetTypeActions;
+	TArray<TSharedPtr<class FAssetTypeActions_Base>> ItemDataAssetTypeActions;
 
 	TSharedPtr<struct FClassBrowseHelper> ClassCache;
 };
diff --git a/Engine/Source/Editor/BehaviorTreeEditor/Public/IBehaviorTreeEditor.h b/Engine/Source/Editor/BehaviorTreeEditor/Public/IBehaviorTreeEditor.h
--- a/Engine/Source/Editor/BehaviorTreeEditor/Public/IBehaviorTreeEditor.h
+++ b/Engine/Source/Editor/BehaviorTreeEditor/Public/IBehaviorTreeEditor.h
@@ -2,12 +2,10 @@
 
 #pragma once
 
-#include "Toolkits/IToolkit.h"
-#include "Toolkits/AssetEditorToolkit.h"
-
+#include "WorkflowCentricApplication.h"
 
 /** BT Editor public interface */
-class IBehaviorTreeEditor : public FAssetEditorToolkit
+class IBehaviorTreeEditor : public FWorkflowCentricApplication
 {
 
 public:
diff --git a/Engine/Source/Editor/EditorStyle/Private/SlateEditorStyle.cpp b/Engine/Source/Editor/EditorStyle/Private/SlateEditorStyle.cpp
--- a/Engine/Source/Editor/EditorStyle/Private/SlateEditorStyle.cpp
+++ b/Engine/Source/Editor/EditorStyle/Private/SlateEditorStyle.cpp
@@ -2815,6 +2815,11 @@ void FSlateEditorStyle::FStyle::SetupGraphEditorStyles()
 			Set( "BTEditor.DebuggerOverlay.FailedTriggerPointer", new IMAGE_BRUSH( "/BehaviorTree/FailedTriggerPointer", FVector2D(48,64)) );
 			Set( "BTEditor.DebuggerOverlay.BreakOnBreakpointPointer", new IMAGE_BRUSH( "Old/Kismet2/IP_Breakpoint", FVector2D(128,96)) );
 
+			Set( "BTEditor.SwitchToBehaviorTreeMode", new IMAGE_BRUSH( "BehaviorTree/BehaviorTreeMode_40x", FVector2D(40,40)));
+			Set( "BTEditor.SwitchToBehaviorTreeMode.Small", new IMAGE_BRUSH( "BehaviorTree/BehaviorTreeMode_16x", FVector2D(16,16)));
+			Set( "BTEditor.SwitchToBlackboardMode", new IMAGE_BRUSH( "BehaviorTree/BlackboardMode_40x", FVector2D(40,40)));
+			Set( "BTEditor.SwitchToBlackboardMode.Small", new IMAGE_BRUSH( "BehaviorTree/BlackboardMode_16x", FVector2D(16,16)));
+
 		}
 
 		// Visible on hover button for transition node
diff --git a/Engine/Source/Editor/Kismet/Private/WorkflowCentricApplication.cpp b/Engine/Source/Editor/Kismet/Private/WorkflowCentricApplication.cpp
--- a/Engine/Source/Editor/Kismet/Private/WorkflowCentricApplication.cpp
+++ b/Engine/Source/Editor/Kismet/Private/WorkflowCentricApplication.cpp
@@ -27,7 +27,7 @@ void FWorkflowCentricApplication::UnregisterTabSpawners(const TSharedRef<class F
 
 FName FWorkflowCentricApplication::GetCurrentMode() const
 {
-	return CurrentAppModePtr->GetModeName();
+	return CurrentAppModePtr.IsValid() ? CurrentAppModePtr->GetModeName() : NAME_None;
 }
 
 void FWorkflowCentricApplication::SetCurrentMode(FName NewMode)
diff --git a/Engine/Source/Editor/UnrealEd/Private/Toolkits/AssetEditorToolkit.cpp b/Engine/Source/Editor/UnrealEd/Private/Toolkits/AssetEditorToolkit.cpp
--- a/Engine/Source/Editor/UnrealEd/Private/Toolkits/AssetEditorToolkit.cpp
+++ b/Engine/Source/Editor/UnrealEd/Private/Toolkits/AssetEditorToolkit.cpp
@@ -211,15 +211,20 @@ FText FAssetEditorToolkit::GetToolkitName() const
 
 	check (EditingObject != NULL);
 
-	const bool bDirtyState = EditingObject->GetOutermost()->IsDirty();
+	return GetDescriptionForObject(EditingObject);
+}
+
+FText FAssetEditorToolkit::GetDescriptionForObject(const UObject* InObject)
+{
+	const bool bDirtyState = InObject->GetOutermost()->IsDirty();
 	FString NameString;
-	if(const AActor* ObjectAsActor = Cast<AActor>(EditingObject))
+	if(const AActor* ObjectAsActor = Cast<AActor>(InObject))
 	{
 		NameString = ObjectAsActor->GetActorLabel();
 	}
 	else
 	{
-		NameString = EditingObject->GetName();
+		NameString = InObject->GetName();
 	}
 
 	FFormatNamedArguments Args;
diff --git a/Engine/Source/Editor/UnrealEd/Public/Toolkits/AssetEditorToolkit.h b/Engine/Source/Editor/UnrealEd/Public/Toolkits/AssetEditorToolkit.h
--- a/Engine/Source/Editor/UnrealEd/Public/Toolkits/AssetEditorToolkit.h
+++ b/Engine/Source/Editor/UnrealEd/Public/Toolkits/AssetEditorToolkit.h
@@ -128,6 +128,13 @@ public:
 	/** True if this actually is editing an asset */
 	bool IsActuallyAnAsset() const;
 
+	/** 
+	 * Gets the text to display in a toolkit titlebar for an object 
+	 * @param	InObject	The object we want a description of
+	 * @return a formatted description of the object state (e.g. "MyObject*")
+	 */
+	static FText GetDescriptionForObject(const UObject* InObject);
+
 protected:
 
 	/**	Returns the single object currently being edited. Asserts if currently editing no object or multiple objects */
diff --git a/Engine/Source/Runtime/AIModule/Classes/BehaviorTree/BehaviorTree.h b/Engine/Source/Runtime/AIModule/Classes/BehaviorTree/BehaviorTree.h
--- a/Engine/Source/Runtime/AIModule/Classes/BehaviorTree/BehaviorTree.h
+++ b/Engine/Source/Runtime/AIModule/Classes/BehaviorTree/BehaviorTree.h
@@ -3,6 +3,7 @@
 
 #include "BTCompositeNode.h"
 #include "BTTaskNode.h"
+#include "Engine/Blueprint.h"
 #include "BehaviorTree.generated.h"
 
 UCLASS(BlueprintType)
@@ -20,6 +21,10 @@ class AIMODULE_API UBehaviorTree : public UObject
 	UPROPERTY()
 	class UEdGraph*	BTGraph;
 
+	/** Info about the graphs we last edited */
+	UPROPERTY()
+	TArray<FEditedDocumentInfo> LastEditedDocuments;
+
 #endif
 
 	/** blackboard asset for this tree */
