Fixed non-unity CIS.

diff --git a/Engine/Source/Editor/BehaviorTreeEditor/Private/BehaviorTreeEditor.cpp b/Engine/Source/Editor/BehaviorTreeEditor/Private/BehaviorTreeEditor.cpp
--- a/Engine/Source/Editor/BehaviorTreeEditor/Private/BehaviorTreeEditor.cpp
+++ b/Engine/Source/Editor/BehaviorTreeEditor/Private/BehaviorTreeEditor.cpp
@@ -20,6 +20,8 @@
 #include "BehaviorTreeEditorTabFactories.h"
 #include "BehaviorTreeEditorCommands.h"
 #include "BehaviorTreeEditorTabs.h"
+#include "BehaviorTree/BehaviorTree.h"
+#include "BehaviorTree/BlackboardData.h"
 
 #define LOCTEXT_NAMESPACE "BehaviorTreeEditor"
 
diff --git a/Engine/Source/Editor/BehaviorTreeEditor/Private/BehaviorTreeEditor.h b/Engine/Source/Editor/BehaviorTreeEditor/Private/BehaviorTreeEditor.h
--- a/Engine/Source/Editor/BehaviorTreeEditor/Private/BehaviorTreeEditor.h
+++ b/Engine/Source/Editor/BehaviorTreeEditor/Private/BehaviorTreeEditor.h
@@ -120,10 +120,10 @@ public:
 	TSharedPtr<class FBehaviorTreeEditorToolbar> GetToolbarBuilder() { return ToolbarBuilder; }
 
 	/** Get the behavior tree we are editing (if any) */
-	UBehaviorTree* GetBehaviorTree() const;
+	class UBehaviorTree* GetBehaviorTree() const;
 
 	/** Get the blackboard we are editing (if any) */
-	UBlackboardData* GetBlackboardData() const;
+	class UBlackboardData* GetBlackboardData() const;
 
 	/** Spawns the tab with the update graph inside */
 	TSharedRef<SWidget> SpawnProperties();
diff --git a/Engine/Source/Editor/BehaviorTreeEditor/Private/BehaviorTreeEditorModes.cpp b/Engine/Source/Editor/BehaviorTreeEditor/Private/BehaviorTreeEditorModes.cpp
--- a/Engine/Source/Editor/BehaviorTreeEditor/Private/BehaviorTreeEditorModes.cpp
+++ b/Engine/Source/Editor/BehaviorTreeEditor/Private/BehaviorTreeEditorModes.cpp
@@ -3,6 +3,8 @@
 #include "BehaviorTreeEditorPrivatePCH.h"
 #include "BehaviorTreeEditorModes.h"
 #include "BehaviorTreeEditorTabs.h"
+#include "BehaviorTreeEditorTabFactories.h"
+#include "BehaviorTreeEditorToolbar.h"
 
 FBehaviorTreeEditorApplicationMode::FBehaviorTreeEditorApplicationMode(TSharedPtr<class FBehaviorTreeEditor> InBehaviorTreeEditor)
 	: FApplicationMode(FBehaviorTreeEditor::BehaviorTreeMode)
diff --git a/Engine/Source/Editor/BehaviorTreeEditor/Private/BehaviorTreeEditorPrivatePCH.h b/Engine/Source/Editor/BehaviorTreeEditor/Private/BehaviorTreeEditorPrivatePCH.h
--- a/Engine/Source/Editor/BehaviorTreeEditor/Private/BehaviorTreeEditorPrivatePCH.h
+++ b/Engine/Source/Editor/BehaviorTreeEditor/Private/BehaviorTreeEditorPrivatePCH.h
@@ -26,3 +26,5 @@
 #include "KismetNodes/KismetNodeInfoContext.h"
 #include "ConnectionDrawingPolicy.h"
 #include "FindInBT.h"
+
+#include "PropertyEditing.h"
diff --git a/Engine/Source/Editor/BehaviorTreeEditor/Private/BehaviorTreeEditorTabFactories.cpp b/Engine/Source/Editor/BehaviorTreeEditor/Private/BehaviorTreeEditorTabFactories.cpp
--- a/Engine/Source/Editor/BehaviorTreeEditor/Private/BehaviorTreeEditorTabFactories.cpp
+++ b/Engine/Source/Editor/BehaviorTreeEditor/Private/BehaviorTreeEditorTabFactories.cpp
@@ -4,6 +4,8 @@
 #include "BehaviorTreeEditorTabFactories.h"
 #include "BehaviorTreeEditor.h"
 #include "BehaviorTreeEditorTabs.h"
+#include "BehaviorTree/BehaviorTree.h"
+#include "BehaviorTree/BlackboardData.h"
 
 #define LOCTEXT_NAMESPACE "BehaviorTreeEditorFactories"
 
