fixed hitch on loading all classes/packages for behavior tree editor

diff --git a/Engine/Source/Editor/BehaviorTreeEditor/Classes/BehaviorTreeEditorTypes.h b/Engine/Source/Editor/BehaviorTreeEditor/Classes/BehaviorTreeEditorTypes.h
--- a/Engine/Source/Editor/BehaviorTreeEditor/Classes/BehaviorTreeEditorTypes.h
+++ b/Engine/Source/Editor/BehaviorTreeEditor/Classes/BehaviorTreeEditorTypes.h
@@ -22,8 +22,8 @@ struct FClassData
 	FClassData(UClass* InClass, const FString& InDeprecatedMessage) : 
 		bIsHidden(0), bHideParent(0), Class(InClass), DeprecatedMessage(InDeprecatedMessage)
 	{}
-	FClassData(const FString& InAssetName, const FString& InGeneratedClassPackage, const FString& InClassName) :
-		bIsHidden(0), bHideParent(0), AssetName(InAssetName), GeneratedClassPackage(InGeneratedClassPackage), ClassName(InClassName) 
+	FClassData(const FString& InAssetName, const FString& InGeneratedClassPackage, const FString& InClassName, UClass* InClass) :
+		bIsHidden(0), bHideParent(0), Class(InClass), AssetName(InAssetName), GeneratedClassPackage(InGeneratedClassPackage), ClassName(InClassName) 
 	{}
 
 	FString ToString() const;
diff --git a/Engine/Source/Editor/BehaviorTreeEditor/Private/BehaviorTreeEditorTypes.cpp b/Engine/Source/Editor/BehaviorTreeEditor/Private/BehaviorTreeEditorTypes.cpp
--- a/Engine/Source/Editor/BehaviorTreeEditor/Private/BehaviorTreeEditorTypes.cpp
+++ b/Engine/Source/Editor/BehaviorTreeEditor/Private/BehaviorTreeEditorTypes.cpp
@@ -192,8 +192,8 @@ bool FClassBrowseHelper::IsHidingParentClass(UClass* Class)
 
 bool FClassBrowseHelper::IsPackageSaved(FName PackageName)
 {
-	UPackage* Package = LoadPackage(NULL, *PackageName.ToString(), LOAD_None);
-	return Package != NULL;
+	const bool bFound = FPackageName::SearchForPackageOnDisk(PackageName.ToString());
+	return bFound;
 }
 
 void FClassBrowseHelper::OnAssetAdded(const class FAssetData& AssetData)
@@ -269,7 +269,11 @@ TSharedPtr<FClassDataNode> FClassBrowseHelper::CreateClassDataNode(const class F
 		Node = MakeShareable(new FClassDataNode);
 		Node->ParentClassName = AssetParentClassName;
 
-		FClassData NewData(AssetData.AssetName.ToString(), AssetData.PackageName.ToString(), AssetClassName);
+		UObject* AssetOb = AssetData.IsAssetLoaded() ? AssetData.GetAsset() : NULL;
+		UBlueprint* AssetBP = Cast<UBlueprint>(AssetOb);
+		UClass* AssetClass = AssetBP ? AssetBP->GeneratedClass : AssetOb ? AssetOb->GetClass() : NULL;
+
+		FClassData NewData(AssetData.AssetName.ToString(), AssetData.PackageName.ToString(), AssetClassName, AssetClass);
 		Node->Data = NewData;
 	}
 
