compilation fix

diff --git a/Engine/Source/Editor/BehaviorTreeEditor/Private/BehaviorTreeEditorTypes.cpp b/Engine/Source/Editor/BehaviorTreeEditor/Private/BehaviorTreeEditorTypes.cpp
--- a/Engine/Source/Editor/BehaviorTreeEditor/Private/BehaviorTreeEditorTypes.cpp
+++ b/Engine/Source/Editor/BehaviorTreeEditor/Private/BehaviorTreeEditorTypes.cpp
@@ -271,7 +271,7 @@ TSharedPtr<FClassDataNode> FClassBrowseHelper::CreateClassDataNode(const class F
 
 		UObject* AssetOb = AssetData.IsAssetLoaded() ? AssetData.GetAsset() : NULL;
 		UBlueprint* AssetBP = Cast<UBlueprint>(AssetOb);
-		UClass* AssetClass = AssetBP ? AssetBP->GeneratedClass : AssetOb ? AssetOb->GetClass() : NULL;
+		UClass* AssetClass = AssetBP ? *AssetBP->GeneratedClass : AssetOb ? AssetOb->GetClass() : NULL;
 
 		FClassData NewData(AssetData.AssetName.ToString(), AssetData.PackageName.ToString(), AssetClassName, AssetClass);
 		Node->Data = NewData;
