Fixed crash opening a new behavior tree

Debugger now checks for a valid graph before dereferencing it. We may not have initialized the graph before the debugger is set up, so we lazy-initialize the root node pointer.

TTP# 338513 - CRITICAL: Regression: BEHAVIOR TREES: CRASH: When opening

diff --git a/Engine/Source/Editor/BehaviorTreeEditor/Private/BehaviorTreeDebugger.cpp b/Engine/Source/Editor/BehaviorTreeEditor/Private/BehaviorTreeDebugger.cpp
--- a/Engine/Source/Editor/BehaviorTreeEditor/Private/BehaviorTreeDebugger.cpp
+++ b/Engine/Source/Editor/BehaviorTreeEditor/Private/BehaviorTreeDebugger.cpp
@@ -36,6 +36,23 @@ FBehaviorTreeDebugger::~FBehaviorTreeDebugger()
 	FBehaviorTreeDelegates::OnDebugSelected.RemoveAll(this);
 }
 
+void FBehaviorTreeDebugger::CacheRootNode()
+{
+	if(RootNode.IsValid() || TreeAsset == nullptr || TreeAsset->BTGraph == nullptr)
+	{
+		return;
+	}
+
+	for (const auto& Node : TreeAsset->BTGraph->Nodes)
+	{
+		RootNode = Cast<UBehaviorTreeGraphNode_Root>(Node);
+		if (RootNode.IsValid())
+		{
+			break;
+		}
+	}
+}
+
 void FBehaviorTreeDebugger::Setup(class UBehaviorTree* InTreeAsset, const class FBehaviorTreeEditor* InEditorOwner, TSharedPtr<SBehaviorTreeDebuggerView> DebuggerView)
 {
 	DebuggerDetails = DebuggerView;
@@ -46,17 +63,7 @@ void FBehaviorTreeDebugger::Setup(class UBehaviorTree* InTreeAsset, const class
 	ActiveBreakpoints.Reset();
 	KnownInstances.Reset();
 
-	if (TreeAsset)
-	{
-		for (int32 i = 0; i < TreeAsset->BTGraph->Nodes.Num(); i++)
-		{
-			RootNode = Cast<UBehaviorTreeGraphNode_Root>(TreeAsset->BTGraph->Nodes[i]);
-			if (RootNode.IsValid())
-			{
-				break;
-			}
-		}
-	}
+	CacheRootNode();
 
 #if USE_BEHAVIORTREE_DEBUGGER
 	if (IsPIESimulating())
@@ -70,6 +77,8 @@ void FBehaviorTreeDebugger::Setup(class UBehaviorTree* InTreeAsset, const class
 
 void FBehaviorTreeDebugger::Refresh()
 {
+	CacheRootNode();
+
 	if (IsPIESimulating())
 	{
 		// make sure is grabs data if currently paused
diff --git a/Engine/Source/Editor/BehaviorTreeEditor/Private/BehaviorTreeDebugger.h b/Engine/Source/Editor/BehaviorTreeEditor/Private/BehaviorTreeDebugger.h
--- a/Engine/Source/Editor/BehaviorTreeEditor/Private/BehaviorTreeDebugger.h
+++ b/Engine/Source/Editor/BehaviorTreeEditor/Private/BehaviorTreeDebugger.h
@@ -21,6 +21,9 @@ public:
 
 	void Setup(class UBehaviorTree* InTreeAsset, const class FBehaviorTreeEditor* InEditorOwner, TSharedPtr<SBehaviorTreeDebuggerView> DebuggerView);
 
+	/** Store the root node for easy access if we have not already */
+	void CacheRootNode();
+
 	bool IsDebuggerReady() const;
 	bool IsDebuggerRunning() const;
 	bool IsShowingCurrentState() const;
