Added improved blackboard editor & replaced simple blackboard debug view.

Added new blackboard view based on MyBlueprint.
Class viewer can now optionally display "Display Names" for classes (this is used to make the names of blackboard entries more friendly).
SPropertyEditorEditInline now displays the class type icon as well as its name.
Blackboard entry details customization now hides the header for entries, so one level of expansion is no longer necessary.
Removed pointer to 'view' in the blackboard debugger & replaced with delegates where appropriate.
Fixed crash issue also mentioned here: https://github.com/EpicGames/UnrealEngine/pull/212
Added new detail customization used to display only the selected entry in the current blackboard.
Fixed crash when running services with no Blackboard (prevent access of NULL blackboard component).

TTP# 337669 - Behaviour Tree Editor: Integrate Blackboard debugging display with Blackboard Editor

diff --git a/Engine/Source/Editor/BehaviorTreeEditor/BehaviorTreeEditor.Build.cs b/Engine/Source/Editor/BehaviorTreeEditor/BehaviorTreeEditor.Build.cs
--- a/Engine/Source/Editor/BehaviorTreeEditor/BehaviorTreeEditor.Build.cs
+++ b/Engine/Source/Editor/BehaviorTreeEditor/BehaviorTreeEditor.Build.cs
@@ -41,6 +41,7 @@ public class BehaviorTreeEditor : ModuleRules
 				"BlueprintGraph",
                 "AIModule",
                 "GameplayDebugger",
+				"ClassViewer"
 			}
 		);
 
diff --git a/Engine/Source/Editor/BehaviorTreeEditor/Private/BehaviorTreeDebugger.cpp b/Engine/Source/Editor/BehaviorTreeEditor/Private/BehaviorTreeDebugger.cpp
--- a/Engine/Source/Editor/BehaviorTreeEditor/Private/BehaviorTreeDebugger.cpp
+++ b/Engine/Source/Editor/BehaviorTreeEditor/Private/BehaviorTreeDebugger.cpp
@@ -6,13 +6,11 @@
 #include "BehaviorTree/BTTaskNode.h"
 #include "BehaviorTree/BTAuxiliaryNode.h"
 #include "BehaviorTreeDelegates.h"
-#include "SBehaviorTreeDebuggerView.h"
 #include "GameplayDebuggingReplicator.h"
 #include "GameplayDebuggingComponent.h"
 
 FBehaviorTreeDebugger::FBehaviorTreeDebugger()
 {
-	EditorOwner = NULL;
 	TreeAsset = NULL;
 	bIsPIEActive = false;
 	bIsCurrentSubtree = false;
@@ -21,6 +19,8 @@ FBehaviorTreeDebugger::FBehaviorTreeDebugger()
 	StepBackIntoIdx = INDEX_NONE;
 	StepBackOverIdx = INDEX_NONE;
 	StepOutIdx = INDEX_NONE;
+	SavedTimestamp = 0.0f;
+	CurrentTimestamp = 0.0f;
 
 	FEditorDelegates::BeginPIE.AddRaw(this, &FBehaviorTreeDebugger::OnBeginPIE);
 	FEditorDelegates::EndPIE.AddRaw(this, &FBehaviorTreeDebugger::OnEndPIE);
@@ -53,10 +53,9 @@ void FBehaviorTreeDebugger::CacheRootNode()
 	}
 }
 
-void FBehaviorTreeDebugger::Setup(class UBehaviorTree* InTreeAsset, const class FBehaviorTreeEditor* InEditorOwner, TSharedPtr<SBehaviorTreeDebuggerView> DebuggerView)
+void FBehaviorTreeDebugger::Setup(class UBehaviorTree* InTreeAsset, TSharedRef<FBehaviorTreeEditor> InEditorOwner)
 {
-	DebuggerDetails = DebuggerView;
-	EditorOwner = (class FBehaviorTreeEditor*)InEditorOwner;
+	EditorOwner = InEditorOwner;
 	TreeAsset = InTreeAsset;
 	DebuggerInstanceIndex = INDEX_NONE;
 	ActiveStepIndex = 0;
@@ -79,7 +78,7 @@ void FBehaviorTreeDebugger::Refresh()
 {
 	CacheRootNode();
 
-	if (IsPIESimulating())
+	if (IsPIESimulating() && IsDebuggerReady())
 	{
 		// make sure is grabs data if currently paused
 		if (IsPlaySessionPaused() && TreeInstance.IsValid())
@@ -183,10 +182,11 @@ bool FBehaviorTreeDebugger::IsTickable() const
 void FBehaviorTreeDebugger::OnBeginPIE(const bool bIsSimulating)
 {
 	bIsPIEActive = true;
-	if(EditorOwner != nullptr)
+	if(EditorOwner.IsValid())
 	{
-		EditorOwner->RegenerateMenusAndToolbars();
-		EditorOwner->DebuggerUpdateGraph();
+		TSharedPtr<FBehaviorTreeEditor> EditorOwnerPtr = EditorOwner.Pin();
+		EditorOwnerPtr->RegenerateMenusAndToolbars();
+		EditorOwnerPtr->DebuggerUpdateGraph();
 	}
 
 	ActiveBreakpoints.Reset();
@@ -207,9 +207,9 @@ void FBehaviorTreeDebugger::OnBeginPIE(const bool bIsSimulating)
 void FBehaviorTreeDebugger::OnEndPIE(const bool bIsSimulating)
 {
 	bIsPIEActive = false;
-	if(EditorOwner != nullptr)
+	if(EditorOwner.IsValid())
 	{
-		EditorOwner->RegenerateMenusAndToolbars();
+		EditorOwner.Pin()->RegenerateMenusAndToolbars();
 	}
 
 	USelection::SelectObjectEvent.RemoveAll(this);
@@ -330,7 +330,10 @@ void FBehaviorTreeDebugger::SetNodeFlags(const struct FBehaviorTreeDebuggerInsta
 	Node->bDebuggerMarkBreakpointTrigger = NodeInstance->GetExecutionIndex() == StoppedOnBreakpointExecutionIndex;
 	if (Node->bDebuggerMarkBreakpointTrigger)
 	{
-		EditorOwner->JumpToNode(Node);
+		if(EditorOwner.IsValid())
+		{
+			EditorOwner.Pin()->JumpToNode(Node);
+		}
 	}
 
 	int32 SearchPathIdx = INDEX_NONE;
@@ -828,7 +831,10 @@ void FBehaviorTreeDebugger::UpdateCurrentStep(int32 PrevStepIdx, int32 NewStepId
 			if (CurInstaceIdx != NewInstanceIdx || 
 				CurStepInfo.InstanceStack[CurInstaceIdx].TreeAsset != NewStepInfo.InstanceStack[NewInstanceIdx].TreeAsset)
 			{
-				EditorOwner->DebuggerSwitchAsset(NewStepInfo.InstanceStack[NewInstanceIdx].TreeAsset);
+				if(EditorOwner.IsValid())
+				{
+					EditorOwner.Pin()->DebuggerSwitchAsset(NewStepInfo.InstanceStack[NewInstanceIdx].TreeAsset);
+				}
 				UpdateCurrentSubtree();
 			}
 		}
@@ -990,24 +996,21 @@ void FBehaviorTreeDebugger::OnBreakpointRemoved(class UBehaviorTreeGraphNode* No
 void FBehaviorTreeDebugger::UpdateDebuggerViewOnInstanceChange()
 {
 #if USE_BEHAVIORTREE_DEBUGGER
-	if (DebuggerDetails.IsValid())
+	UBlackboardData* BBAsset = EditorOwner.IsValid() ? EditorOwner.Pin()->GetBlackboardData() : nullptr;
+	if (TreeInstance.IsValid() &&
+		TreeInstance->DebuggerSteps.IsValidIndex(ActiveStepIndex) &&
+		TreeInstance->DebuggerSteps[ActiveStepIndex].InstanceStack.IsValidIndex(DebuggerInstanceIndex))
 	{
-		UBlackboardData* BBAsset = NULL;
-		if (TreeInstance.IsValid() &&
-			TreeInstance->DebuggerSteps.IsValidIndex(ActiveStepIndex) &&
-			TreeInstance->DebuggerSteps[ActiveStepIndex].InstanceStack.IsValidIndex(DebuggerInstanceIndex))
+		const FBehaviorTreeDebuggerInstance& ShowInstance = TreeInstance->DebuggerSteps[ActiveStepIndex].InstanceStack[DebuggerInstanceIndex];
+		if (ShowInstance.TreeAsset)
 		{
-			const FBehaviorTreeDebuggerInstance& ShowInstance = TreeInstance->DebuggerSteps[ActiveStepIndex].InstanceStack[DebuggerInstanceIndex];
-			if (ShowInstance.TreeAsset)
-			{
-				BBAsset = ShowInstance.TreeAsset->BlackboardAsset;
-			}
+			BBAsset = ShowInstance.TreeAsset->BlackboardAsset;
 		}
+	}
 
-		DebuggerDetails->UpdateBlackboard(BBAsset);
+	OnDebuggedBlackboardChangedEvent.Broadcast(BBAsset);
 
-		Refresh();
-	}
+	Refresh();
 #endif
 }
 
@@ -1015,14 +1018,13 @@ void FBehaviorTreeDebugger::UpdateDebuggerViewOnStepChange()
 {
 #if USE_BEHAVIORTREE_DEBUGGER
 	if (IsDebuggerRunning() &&
-		DebuggerDetails.IsValid() && TreeInstance.IsValid() &&
+		TreeInstance.IsValid() &&
 		TreeInstance->DebuggerSteps.IsValidIndex(ActiveStepIndex))
 	{
 		const FBehaviorTreeExecutionStep& ShowStep = TreeInstance->DebuggerSteps[ActiveStepIndex];
 
-		DebuggerDetails->SavedTimestamp = ShowStep.TimeStamp;
-		DebuggerDetails->SavedValues = ShowStep.BlackboardValues;
-		DebuggerDetails->bCanShowCurrentState = IsShowingCurrentState();
+		SavedTimestamp = ShowStep.TimeStamp;
+		SavedValues = ShowStep.BlackboardValues;
 	}
 #endif
 }
@@ -1030,18 +1032,51 @@ void FBehaviorTreeDebugger::UpdateDebuggerViewOnStepChange()
 void FBehaviorTreeDebugger::UpdateDebuggerViewOnTick()
 {
 #if USE_BEHAVIORTREE_DEBUGGER
-	if (IsDebuggerRunning() && 
-		DebuggerDetails.IsValid() && TreeInstance.IsValid() &&
-		DebuggerDetails->IsShowingCurrentState())
+	if (IsDebuggerRunning() && TreeInstance.IsValid())
 	{
 		const float GameTime = GEditor && GEditor->PlayWorld ? GEditor->PlayWorld->GetTimeSeconds() : 0.0f;
-		DebuggerDetails->CurrentTimestamp = GameTime;
+		CurrentTimestamp = GameTime;
+
+		TreeInstance->StoreDebuggerBlackboard(CurrentValues);
+	}
+#endif
+}
 
-		TreeInstance->StoreDebuggerBlackboard(DebuggerDetails->CurrentValues);
+FText FBehaviorTreeDebugger::FindValueForKey(const FName& InKeyName, bool bUseCurrentState) const
+{
+#if USE_BEHAVIORTREE_DEBUGGER
+	if (IsDebuggerRunning() &&
+		TreeInstance.IsValid())
+	{
+		FBehaviorTreeExecutionStep* ShowStep = nullptr;
+		if(bUseCurrentState && TreeInstance->DebuggerSteps.Num() > 0)
+		{
+			ShowStep = &TreeInstance->DebuggerSteps[TreeInstance->DebuggerSteps.Num() - 1];
+		}
+		else if(TreeInstance->DebuggerSteps.IsValidIndex(ActiveStepIndex))
+		{
+			ShowStep = &TreeInstance->DebuggerSteps[ActiveStepIndex];
+		}
+
+		if(ShowStep != nullptr)
+		{
+			const FString* FindValue = ShowStep->BlackboardValues.Find(InKeyName);
+			if(FindValue != nullptr)
+			{
+				return FText::FromString(*FindValue);
+			}	
+		}
 	}
+
+	return FText();
 #endif
 }
 
+float FBehaviorTreeDebugger::GetTimeStamp(bool bUseCurrentState) const
+{
+	return bUseCurrentState ? CurrentTimestamp : SavedTimestamp;
+}
+
 FString FBehaviorTreeDebugger::GetDebuggedInstanceDesc() const
 {
 	return TreeInstance.IsValid() ? 
diff --git a/Engine/Source/Editor/BehaviorTreeEditor/Private/BehaviorTreeDebugger.h b/Engine/Source/Editor/BehaviorTreeEditor/Private/BehaviorTreeDebugger.h
--- a/Engine/Source/Editor/BehaviorTreeEditor/Private/BehaviorTreeDebugger.h
+++ b/Engine/Source/Editor/BehaviorTreeEditor/Private/BehaviorTreeDebugger.h
@@ -2,8 +2,6 @@
 
 #pragma once
 
-#include "SBehaviorTreeDebuggerView.h"
-
 class FBehaviorTreeDebugger : public FTickableGameObject
 {
 public:
@@ -19,7 +17,7 @@ public:
 	/** Refresh the debugging information we are displaying (only when paused, as Tick() updates when running) */
 	void Refresh();
 
-	void Setup(class UBehaviorTree* InTreeAsset, const class FBehaviorTreeEditor* InEditorOwner, TSharedPtr<SBehaviorTreeDebuggerView> DebuggerView);
+	void Setup(class UBehaviorTree* InTreeAsset, TSharedRef<class FBehaviorTreeEditor> InEditorOwner);
 
 	/** Store the root node for easy access if we have not already */
 	void CacheRootNode();
@@ -66,10 +64,25 @@ public:
 	bool HasContinuousNextStep() const;
 	bool HasContinuousPrevStep() const;
 
+	/**
+	 * Find a (display) value for a given key.
+	 * @param	InKeyName			Key to find a value for
+	 * @param	bUseCurrentState	Whether to use the current (present) state or the state at the active step index
+	 * @return the value to display to the user.
+	 */
+	FText FindValueForKey(const FName& InKeyName, bool bUseCurrentState) const;
+
+	/** Gets the timestamp to be displayed, either current or saved */
+	float GetTimeStamp(bool bUseCurrentState) const;
+
+	/** Delegate fired when the debugged blackboard is changed */
+	DECLARE_EVENT_OneParam(FBehaviorTreeDebugger, FOnDebuggedBlackboardChanged, UBlackboardData*);
+	FOnDebuggedBlackboardChanged& OnDebuggedBlackboardChanged() { return OnDebuggedBlackboardChangedEvent; }
+
 private:
 
 	/** owning editor */
-	class FBehaviorTreeEditor* EditorOwner;
+	TWeakPtr<FBehaviorTreeEditor> EditorOwner;
 
 	/** asset for debugging */
 	class UBehaviorTree* TreeAsset;
@@ -80,9 +93,6 @@ private:
 	/** instance for debugging */
 	TWeakObjectPtr<class UBehaviorTreeComponent> TreeInstance;
 
-	/** widget with details to show in editor */
-	TSharedPtr<SBehaviorTreeDebuggerView> DebuggerDetails;
-
 	/** matching debugger instance index from component's stack */
 	int32 DebuggerInstanceIndex;
 
@@ -114,6 +124,14 @@ private:
 	/** execution index of node that caused activated the breakpoint */
 	uint16 StoppedOnBreakpointExecutionIndex;
 
+	/** Lookup of currently debugged blackboard values */
+	TMap<FName, FString> SavedValues;
+	TMap<FName, FString> CurrentValues;
+
+	/** Debugger timestamps */
+	float SavedTimestamp;
+	float CurrentTimestamp;
+
 	/** set value of DebuggerInstanceIndex variable */
 	void UpdateDebuggerInstance();
 
@@ -173,4 +191,7 @@ private:
 
 	/** updates button states */
 	void UpdateAvailableActions();
+
+	/** Delegate fired when the debugged blackboard is changed */
+	FOnDebuggedBlackboardChanged OnDebuggedBlackboardChangedEvent;
 };
diff --git a/Engine/Source/Editor/BehaviorTreeEditor/Private/BehaviorTreeEditor.cpp b/Engine/Source/Editor/BehaviorTreeEditor/Private/BehaviorTreeEditor.cpp
--- a/Engine/Source/Editor/BehaviorTreeEditor/Private/BehaviorTreeEditor.cpp
+++ b/Engine/Source/Editor/BehaviorTreeEditor/Private/BehaviorTreeEditor.cpp
@@ -22,6 +22,9 @@
 #include "BehaviorTreeEditorTabs.h"
 #include "BehaviorTree/BehaviorTree.h"
 #include "BehaviorTree/BlackboardData.h"
+#include "DetailCustomizations/BlackboardDataDetails.h"
+#include "SBehaviorTreeBlackboardView.h"
+#include "SBehaviorTreeBlackboardEditor.h"
 
 #define LOCTEXT_NAMESPACE "BehaviorTreeEditor"
 
@@ -75,6 +78,14 @@ void FBehaviorTreeEditor::PostUndo(bool bSuccess)
 			CurrentGraphEditor->ClearSelectionSet();
 			CurrentGraphEditor->NotifyGraphChanged();
 		}
+		if(BlackboardView.IsValid())
+		{
+			BlackboardView->SetObject(GetBlackboardData());
+		}
+		if(BlackboardEditor.IsValid())
+		{
+			BlackboardEditor->SetObject(GetBlackboardData());
+		}
 		FSlateApplication::Get().DismissAllMenus();
 	}
 }
@@ -90,10 +101,33 @@ void FBehaviorTreeEditor::PostRedo(bool bSuccess)
 			CurrentGraphEditor->ClearSelectionSet();
 			CurrentGraphEditor->NotifyGraphChanged();
 		}
+		if(BlackboardView.IsValid())
+		{
+			BlackboardView->SetObject(GetBlackboardData());
+		}
+		if(BlackboardEditor.IsValid())
+		{
+			BlackboardEditor->SetObject(GetBlackboardData());
+		}
 		FSlateApplication::Get().DismissAllMenus();
 	}
 }
 
+void FBehaviorTreeEditor::NotifyPostChange( const FPropertyChangedEvent& PropertyChangedEvent, UProperty* PropertyThatChanged )
+{
+	if(PropertyChangedEvent.ChangeType != EPropertyChangeType::Interactive)
+	{
+		if(BlackboardView.IsValid())
+		{
+			BlackboardView->SetObject(GetBlackboardData());
+		}
+		if(BlackboardEditor.IsValid())
+		{
+			BlackboardEditor->SetObject(GetBlackboardData());
+		}
+	}
+}
+
 void FBehaviorTreeEditor::RegisterTabSpawners(const TSharedRef<class FTabManager>& TabManager)
 {
 	DocumentManager->SetTabManager(TabManager);
@@ -164,6 +198,7 @@ void FBehaviorTreeEditor::InitBehaviorTreeEditor( const EToolkitMode::Type Mode,
 		FGraphEditorCommands::Register();
 		FBTCommonCommands::Register();
 		FBTDebuggerCommands::Register();
+		FBTBlackboardCommands::Register();
 
 		const TSharedRef<FTabManager::FLayout> DummyLayout = FTabManager::NewLayout("NullLayout")->AddArea(FTabManager::NewPrimaryArea());
 		const bool bCreateDefaultStandaloneMenu = true;
@@ -175,19 +210,35 @@ void FBehaviorTreeEditor::InitBehaviorTreeEditor( const EToolkitMode::Type Mode,
 		CreateInternalWidgets();
 
 		Debugger = MakeShareable(new FBehaviorTreeDebugger);
-		Debugger->Setup(BehaviorTree, this, DebuggerView);
+		Debugger->Setup(BehaviorTree, SharedThis(this));
+		Debugger->OnDebuggedBlackboardChanged().AddSP(this, &FBehaviorTreeEditor::HandleDebuggedBlackboardChanged);
 		BindDebuggerToolbarCommands();
 
-		FBehaviorTreeEditorModule& BehaviorTreeEditorModule = FModuleManager::GetModuleChecked<FBehaviorTreeEditorModule>( "BehaviorTreeEditor" );
+		FBehaviorTreeEditorModule& BehaviorTreeEditorModule = FModuleManager::LoadModuleChecked<FBehaviorTreeEditorModule>( "BehaviorTreeEditor" );
 		AddMenuExtender(BehaviorTreeEditorModule.GetMenuExtensibilityManager()->GetAllExtenders(GetToolkitCommands(), GetEditingObjects()));
 
 		AddApplicationMode(BehaviorTreeMode, MakeShareable(new FBehaviorTreeEditorApplicationMode(SharedThis(this))));
 		AddApplicationMode(BlackboardMode, MakeShareable(new FBlackboardEditorApplicationMode(SharedThis(this))));
+
+		BlackboardView = SNew(SBehaviorTreeBlackboardView, GetBlackboardData())
+			.OnGetDebugKeyValue(this, &FBehaviorTreeEditor::HandleGetDebugKeyValue)
+			.OnIsDebuggerReady(this, &FBehaviorTreeEditor::IsDebuggerReady)
+			.OnIsDebuggerPaused(this, &FBehaviorTreeEditor::IsDebuggerPaused)
+			.OnGetDebugTimeStamp(this, &FBehaviorTreeEditor::HandleGetDebugTimeStamp)
+			.OnGetDisplayCurrentState(this, &FBehaviorTreeEditor::HandleGetDisplayCurrentState);
+
+		BlackboardEditor = SNew(SBehaviorTreeBlackboardEditor, GetToolkitCommands(), GetBlackboardData())
+				.OnEntrySelected(this, &FBehaviorTreeEditor::HandleBlackboardEntrySelected)
+				.OnGetDebugKeyValue(this, &FBehaviorTreeEditor::HandleGetDebugKeyValue)
+				.OnIsDebuggerReady(this, &FBehaviorTreeEditor::IsDebuggerReady)
+				.OnIsDebuggerPaused(this, &FBehaviorTreeEditor::IsDebuggerPaused)
+				.OnGetDebugTimeStamp(this, &FBehaviorTreeEditor::HandleGetDebugTimeStamp)
+				.OnGetDisplayCurrentState(this, &FBehaviorTreeEditor::HandleGetDisplayCurrentState);
 	}
 	else
 	{
 		check(Debugger.IsValid());
-		Debugger->Setup(BehaviorTree, this, DebuggerView);
+		Debugger->Setup(BehaviorTree, SharedThis(this));
 	}
 
 	if(BehaviorTreeToEdit != nullptr)
@@ -236,13 +287,14 @@ void FBehaviorTreeEditor::RestoreBehaviorTree()
 		GraphEditor->SetViewLocation(BehaviorTree->LastEditedDocuments[0].SavedViewOffset, BehaviorTree->LastEditedDocuments[0].SavedZoomAmount);
 	}
 
+	const bool bIncreaseVersionNum = false;
 	if(bNewGraph)
 	{
-		const bool bIncreaseVersionNum = false;
 		MyGraph->UpdateAsset(UBehaviorTreeGraph::ClearDebuggerFlags, bIncreaseVersionNum);
 	}
 	else
 	{
+		MyGraph->UpdateAsset(UBehaviorTreeGraph::SkipDebuggerFlags, bIncreaseVersionNum);
 		RefreshDebugger();
 	}
 
@@ -262,11 +314,91 @@ void FBehaviorTreeEditor::SaveEditedObjectState()
 	DocumentManager->SaveAllState();
 }
 
+void FBehaviorTreeEditor::HandleBlackboardEntrySelected(const FBlackboardEntry* BlackboardEntry, bool bIsInherited)
+{
+	// refresh details view
+	const bool bForceRefresh = true;
+	BlackboardDetailsView->SetObject(GetBlackboardData(), bForceRefresh);
+}
+
+int32 FBehaviorTreeEditor::HandleGetSelectedBlackboardItemIndex(bool& bIsInherited)
+{
+	if(BlackboardEditor.IsValid())
+	{
+		return BlackboardEditor->GetSelectedEntryIndex(bIsInherited);
+	}
+
+	return INDEX_NONE;
+}
+
+FText FBehaviorTreeEditor::HandleGetDebugKeyValue(const FName& InKeyName, bool bUseCurrentState) const
+{
+	if(IsDebuggerReady())
+	{
+		return Debugger->FindValueForKey(InKeyName, bUseCurrentState);
+	}
+
+	return FText();
+}
+
+float FBehaviorTreeEditor::HandleGetDebugTimeStamp(bool bUseCurrentState) const
+{
+	if(IsDebuggerReady())
+	{
+		return Debugger->GetTimeStamp(bUseCurrentState);
+	}
+
+	return 0.0f;
+}
+
+void FBehaviorTreeEditor::HandleDebuggedBlackboardChanged(UBlackboardData* InBlackboardData)
+{
+	if(BlackboardView.IsValid())
+	{
+		BlackboardView->SetObject(InBlackboardData);
+	}
+	if(BlackboardEditor.IsValid())
+	{
+		BlackboardEditor->SetObject(InBlackboardData);
+	}
+}
+
+bool FBehaviorTreeEditor::HandleGetDisplayCurrentState() const
+{
+	if(IsDebuggerReady())
+	{
+		return Debugger->IsShowingCurrentState();
+	}
+
+	return false;
+}
+
+void FBehaviorTreeEditor::GetBlackboardSelectionInfo(int32& OutSelectionIndex, bool& bOutIsInherited) const
+{
+	OutSelectionIndex = CurrentBlackboardEntryIndex;
+	bOutIsInherited = bIsCurrentBlackboardEntryInherited;
+}
+
 bool FBehaviorTreeEditor::IsDebuggerReady() const
 {
 	return Debugger.IsValid() && Debugger->IsDebuggerReady();
 }
 
+bool FBehaviorTreeEditor::IsDebuggerPaused() const
+{
+	return IsDebuggerReady() && GUnrealEd->PlayWorld && GUnrealEd->PlayWorld->bDebugPauseExecution;
+}
+
+bool FBehaviorTreeEditor::CanEditWithDebuggerActive() const
+{
+	if(Debugger.IsValid())
+	{
+		return !Debugger->IsDebuggerReady();
+	}
+
+	return true;
+}
+
 EVisibility FBehaviorTreeEditor::GetDebuggerDetailsVisibility() const
 {
 	return Debugger.IsValid() && Debugger->IsDebuggerRunning() ? EVisibility::Visible : EVisibility::Collapsed;
@@ -520,26 +652,42 @@ TSharedRef<SWidget> FBehaviorTreeEditor::SpawnProperties()
 					.Text(FText::Format(LOCTEXT("AbortModeHighlight", "Nodes aborted by mode: {0}"), LOCTEXT("AbortPrioritySelf", "Self")))
 				]
 			]
-		]
-		+SVerticalBox::Slot()
-		.FillHeight(1.0f)
-		.HAlign(HAlign_Fill)
-		.VAlign(VAlign_Bottom)
-		.Padding(5.0f)
-		[
-			SNew(SSeparator)
-			.Orientation( Orient_Horizontal )
-			.Visibility(this, &FBehaviorTreeEditor::GetDebuggerDetailsVisibility)
-		]
-		+SVerticalBox::Slot()
-		.AutoHeight()
-		.HAlign(HAlign_Left)
-		.VAlign(VAlign_Top)
-		[
-			DebuggerView.ToSharedRef()
 		];
 }
 
+TSharedRef<SWidget> FBehaviorTreeEditor::SpawnBlackboardView()
+{
+	return BlackboardView.ToSharedRef();
+}
+
+TSharedRef<SWidget> FBehaviorTreeEditor::SpawnBlackboardEditor()
+{
+	return BlackboardEditor.ToSharedRef();
+}
+
+TSharedRef<SWidget> FBehaviorTreeEditor::SpawnBlackboardDetails()
+{
+	const bool bIsUpdatable = false;
+	const bool bAllowFavorites = true;
+	const bool bIsLockable = false;
+	const bool bAllowSearch = true;
+	const bool bObjectsUseNameArea = false;
+	const bool bHideSelectionTip = true;
+	FPropertyEditorModule& PropertyEditorModule = FModuleManager::GetModuleChecked<FPropertyEditorModule>( "PropertyEditor" );
+	FDetailsViewArgs DetailsViewArgs( bIsUpdatable, bIsLockable, bAllowSearch, bObjectsUseNameArea, bHideSelectionTip );
+	DetailsViewArgs.NotifyHook = this;
+	BlackboardDetailsView = PropertyEditorModule.CreateDetailView( DetailsViewArgs );
+
+	FOnGetSelectedBlackboardItemIndex OnGetSelectedBlackboardItemIndex = FOnGetSelectedBlackboardItemIndex::CreateSP(this, &FBehaviorTreeEditor::HandleGetSelectedBlackboardItemIndex);
+	FOnGetDetailCustomizationInstance LayoutVariableDetails = FOnGetDetailCustomizationInstance::CreateStatic(&FBlackboardDataDetails::MakeInstance, OnGetSelectedBlackboardItemIndex);
+	BlackboardDetailsView->RegisterInstancedCustomPropertyLayout(UBlackboardData::StaticClass(), LayoutVariableDetails);
+
+	BlackboardDetailsView->SetObject(GetBlackboardData());
+	BlackboardDetailsView->SetEnabled(TAttribute<bool>::Create(TAttribute<bool>::FGetter::CreateSP(this, &FBehaviorTreeEditor::CanEditWithDebuggerActive)));
+
+	return BlackboardDetailsView.ToSharedRef();
+}
+
 void FBehaviorTreeEditor::OnSelectedNodesChanged(const TSet<class UObject*>& NewSelection)
 {
 	TArray<UObject*> Selection;
@@ -757,8 +905,6 @@ void FBehaviorTreeEditor::CreateInternalWidgets()
 	DetailsView->SetIsPropertyVisibleDelegate(FIsPropertyVisible::CreateSP(this, &FBehaviorTreeEditor::IsPropertyVisible));
 	DetailsView->SetIsPropertyEditingEnabledDelegate(FIsPropertyEditingEnabled::CreateSP(this, &FBehaviorTreeEditor::IsPropertyEditable));
 	DetailsView->OnFinishedChangingProperties().AddSP(this, &FBehaviorTreeEditor::OnFinishedChangingProperties);
-
-	DebuggerView = SNew(SBehaviorTreeDebuggerView);
 }
 
 void FBehaviorTreeEditor::ExtendMenu()
@@ -925,26 +1071,36 @@ void FBehaviorTreeEditor::OnFinishedChangingProperties(const FPropertyChangedEve
 	}
 
 	// update abort range highlight when changing decorator's flow abort mode
-	if (PropertyChangedEvent.Property && PropertyChangedEvent.Property->GetFName() == TEXT("FlowAbortMode"))
+	if (PropertyChangedEvent.Property)
 	{
-		bShowDecoratorRangeLower = false;
-		bShowDecoratorRangeSelf = false;
-
-		FGraphPanelSelectionSet CurrentSelection = GetSelectedNodes();
-		if (CurrentSelection.Num() == 1)
+		if(PropertyChangedEvent.Property->GetFName() == TEXT("FlowAbortMode"))
 		{
-			for (FGraphPanelSelectionSet::TConstIterator It(CurrentSelection); It; ++It)
+			bShowDecoratorRangeLower = false;
+			bShowDecoratorRangeSelf = false;
+
+			FGraphPanelSelectionSet CurrentSelection = GetSelectedNodes();
+			if (CurrentSelection.Num() == 1)
 			{
-				UBehaviorTreeGraphNode_Decorator* DecoratorNode = Cast<UBehaviorTreeGraphNode_Decorator>(*It);
-				if (DecoratorNode)
+				for (FGraphPanelSelectionSet::TConstIterator It(CurrentSelection); It; ++It)
 				{
-					FAbortDrawHelper Mode0, Mode1;
-					GetAbortModePreview(Cast<UBTDecorator>(DecoratorNode->NodeInstance), Mode0, Mode1);
+					UBehaviorTreeGraphNode_Decorator* DecoratorNode = Cast<UBehaviorTreeGraphNode_Decorator>(*It);
+					if (DecoratorNode)
+					{
+						FAbortDrawHelper Mode0, Mode1;
+						GetAbortModePreview(Cast<UBTDecorator>(DecoratorNode->NodeInstance), Mode0, Mode1);
 
-					UBehaviorTreeGraph* MyGraph = Cast<UBehaviorTreeGraph>(BehaviorTree->BTGraph);
-					MyGraph->UpdateAbortHighlight(Mode0, Mode1);
-				}
-			}			
+						UBehaviorTreeGraph* MyGraph = Cast<UBehaviorTreeGraph>(BehaviorTree->BTGraph);
+						MyGraph->UpdateAbortHighlight(Mode0, Mode1);
+					}
+				}			
+			}
+		}
+		else if(PropertyChangedEvent.Property->GetFName() == TEXT("BlackboardAsset"))
+		{
+			if(BlackboardView.IsValid())
+			{
+				BlackboardView->SetObject(GetBlackboardData());
+			}
 		}
 	}
 
diff --git a/Engine/Source/Editor/BehaviorTreeEditor/Private/BehaviorTreeEditor.h b/Engine/Source/Editor/BehaviorTreeEditor/Private/BehaviorTreeEditor.h
--- a/Engine/Source/Editor/BehaviorTreeEditor/Private/BehaviorTreeEditor.h
+++ b/Engine/Source/Editor/BehaviorTreeEditor/Private/BehaviorTreeEditor.h
@@ -5,7 +5,11 @@
 #include "IBehaviorTreeEditor.h"
 #include "Toolkits/AssetEditorToolkit.h"
 
-class FBehaviorTreeEditor : public IBehaviorTreeEditor, public FEditorUndoClient
+class UBehaviorTree;
+class UBlackboardData;
+struct FBlackboardEntry;
+
+class FBehaviorTreeEditor : public IBehaviorTreeEditor, public FEditorUndoClient, public FNotifyHook
 {
 public:
 	FBehaviorTreeEditor();
@@ -37,6 +41,10 @@ public:
 	virtual void	PostRedo(bool bSuccess) override;
 	// End of FEditorUndoClient
 
+	// Begin FNotifyHook Interface
+	virtual void NotifyPostChange( const FPropertyChangedEvent& PropertyChangedEvent, UProperty* PropertyThatChanged ) override;
+	// End of FNotifyHook
+
 	// Delegates
 	void OnNodeDoubleClicked(class UEdGraphNode* Node);
 	void OnGraphEditorFocused(const TSharedRef<SGraphEditor>& InGraphEditor);
@@ -87,6 +95,12 @@ public:
 	void UpdateToolbar();
 	bool IsDebuggerReady() const;
 
+	/** Get whether the debugger is currently running and the PIE session is paused */
+	bool IsDebuggerPaused() const;
+
+	/** Get whether we can edit the tree/blackboard with the debugger active */
+	bool CanEditWithDebuggerActive() const;
+
 	TSharedRef<class SWidget> OnGetDebuggerActorsMenu();
 	void OnDebuggerActorSelected(TWeakObjectPtr<UBehaviorTreeComponent> InstanceToDebug);
 	FString GetDebuggerActorDesc() const;
@@ -120,10 +134,10 @@ public:
 	TSharedPtr<class FBehaviorTreeEditorToolbar> GetToolbarBuilder() { return ToolbarBuilder; }
 
 	/** Get the behavior tree we are editing (if any) */
-	class UBehaviorTree* GetBehaviorTree() const;
+	UBehaviorTree* GetBehaviorTree() const;
 
 	/** Get the blackboard we are editing (if any) */
-	class UBlackboardData* GetBlackboardData() const;
+	UBlackboardData* GetBlackboardData() const;
 
 	/** Spawns the tab with the update graph inside */
 	TSharedRef<SWidget> SpawnProperties();
@@ -131,6 +145,15 @@ public:
 	/** Spawns the search tab */
 	TSharedRef<SWidget> SpawnSearch();
 
+	/** Spawn blackboard details tab */
+	TSharedRef<SWidget> SpawnBlackboardDetails();
+
+	/** Spawn blackboard view tab */
+	TSharedRef<SWidget> SpawnBlackboardView();
+
+	/** Spawn blackboard editor tab */
+	TSharedRef<SWidget> SpawnBlackboardEditor();
+
 	// @todo This is a hack for now until we reconcile the default toolbar with application modes [duplicated from counterpart in Blueprint Editor]
 	void RegisterToolbarTab(const TSharedRef<class FTabManager>& TabManager);
 
@@ -140,6 +163,27 @@ public:
 	/** Save the graph state for later editing */
 	void SaveEditedObjectState();
 
+	/** Delegate handler for selection in the blackboard entry list */
+	void HandleBlackboardEntrySelected(const FBlackboardEntry* BlackboardEntry, bool bIsInherited);
+
+	/** Delegate handler used to retrieve current blackboard selection */
+	int32 HandleGetSelectedBlackboardItemIndex(bool& bOutIsInherited);
+
+	/** Delegate handler for displaying debugger values */
+	FText HandleGetDebugKeyValue(const FName& InKeyName, bool bUseCurrentState) const;
+
+	/** Delegate handler for retrieving timestamp to display */
+	float HandleGetDebugTimeStamp(bool bUseCurrentState) const;
+
+	/** Delegate handler for when the debugged blackboard changes */
+	void HandleDebuggedBlackboardChanged(UBlackboardData* InObject);
+
+	/** Delegate handler for determining whether to display the current state */
+	bool HandleGetDisplayCurrentState() const;
+
+	/** Get the currently selected blackboard entry */
+	void GetBlackboardSelectionInfo(int32& OutSelectionIndex, bool& bOutIsInherited) const;
+
 protected:
 	/** Called when "Save" is clicked for this asset */
 	virtual void SaveAsset_Execute() override;
@@ -186,7 +230,6 @@ private:
 
 	/** Property View */
 	TSharedPtr<class IDetailsView> DetailsView;
-	TSharedPtr<class SBehaviorTreeDebuggerView> DebuggerView;
 
 	/** The command list for this editor */
 	TSharedPtr<FUICommandList> GraphEditorCommands;
@@ -204,6 +247,21 @@ private:
 
 	TSharedPtr<class FBehaviorTreeEditorToolbar> ToolbarBuilder;
 
+	/** The details view we use to display the blackboard */
+	TSharedPtr<IDetailsView> BlackboardDetailsView;
+
+	/** The blackboard view widget */
+	TSharedPtr<class SBehaviorTreeBlackboardView> BlackboardView;
+
+	/** The blackboard editor widget */
+	TSharedPtr<class SBehaviorTreeBlackboardEditor> BlackboardEditor;
+
+	/** The current blackboard selection index, stored here so it can be accessed by our details customization */
+	int32 CurrentBlackboardEntryIndex;
+
+	/** Whether the current selection is inherited, stored here so it can be accessed by our details customization */
+	bool bIsCurrentBlackboardEntryInherited;
+
 public:
 	/** Modes in mode switcher */
 	static const FName BehaviorTreeMode;
diff --git a/Engine/Source/Editor/BehaviorTreeEditor/Private/BehaviorTreeEditorCommands.cpp b/Engine/Source/Editor/BehaviorTreeEditor/Private/BehaviorTreeEditorCommands.cpp
--- a/Engine/Source/Editor/BehaviorTreeEditor/Private/BehaviorTreeEditorCommands.cpp
+++ b/Engine/Source/Editor/BehaviorTreeEditor/Private/BehaviorTreeEditorCommands.cpp
@@ -31,6 +31,19 @@ void FBTDebuggerCommands::RegisterCommands()
 	UI_COMMAND(PausePlaySession, "Pause", "Pause simulation", EUserInterfaceActionType::Button, FInputGesture());
 	UI_COMMAND(ResumePlaySession, "Resume", "Resume simulation", EUserInterfaceActionType::Button, FInputGesture() );
 	UI_COMMAND(StopPlaySession, "Stop", "Stop simulation", EUserInterfaceActionType::Button, FInputGesture());
+
+	UI_COMMAND(CurrentValues, "Current", "View current values", EUserInterfaceActionType::RadioButton, FInputGesture());
+	UI_COMMAND(SavedValues, "Saved", "View saved values", EUserInterfaceActionType::RadioButton, FInputGesture());
+}
+
+FBTBlackboardCommands::FBTBlackboardCommands() 
+	: TCommands<FBTBlackboardCommands>("BTEditor.Blackboard", LOCTEXT("Blackboard", "Blackboard"), NAME_None, FEditorStyle::GetStyleSetName())
+{
+}
+
+void FBTBlackboardCommands::RegisterCommands()
+{
+	UI_COMMAND(DeleteEntry, "Delete", "Delete this blackboard entry", EUserInterfaceActionType::Button, FInputGesture(EKeys::Platform_Delete));
 }
 
 #undef LOCTEXT_NAMESPACE
diff --git a/Engine/Source/Editor/BehaviorTreeEditor/Private/BehaviorTreeEditorCommands.h b/Engine/Source/Editor/BehaviorTreeEditor/Private/BehaviorTreeEditorCommands.h
--- a/Engine/Source/Editor/BehaviorTreeEditor/Private/BehaviorTreeEditorCommands.h
+++ b/Engine/Source/Editor/BehaviorTreeEditor/Private/BehaviorTreeEditorCommands.h
@@ -30,6 +30,20 @@ public:
 	TSharedPtr<FUICommandInfo> ResumePlaySession;
 	TSharedPtr<FUICommandInfo> StopPlaySession;
 
+	TSharedPtr<FUICommandInfo> CurrentValues;
+	TSharedPtr<FUICommandInfo> SavedValues;
+
+	/** Initialize commands */
+	virtual void RegisterCommands() override;
+};
+
+class FBTBlackboardCommands : public TCommands<FBTBlackboardCommands>
+{
+public:
+	FBTBlackboardCommands();
+
+	TSharedPtr<FUICommandInfo> DeleteEntry;
+
 	/** Initialize commands */
 	virtual void RegisterCommands() override;
 };
diff --git a/Engine/Source/Editor/BehaviorTreeEditor/Private/BehaviorTreeEditorModes.cpp b/Engine/Source/Editor/BehaviorTreeEditor/Private/BehaviorTreeEditorModes.cpp
--- a/Engine/Source/Editor/BehaviorTreeEditor/Private/BehaviorTreeEditorModes.cpp
+++ b/Engine/Source/Editor/BehaviorTreeEditor/Private/BehaviorTreeEditorModes.cpp
@@ -13,8 +13,9 @@ FBehaviorTreeEditorApplicationMode::FBehaviorTreeEditorApplicationMode(TSharedPt
 
 	BehaviorTreeEditorTabFactories.RegisterFactory(MakeShareable(new FBehaviorTreeDetailsSummoner(InBehaviorTreeEditor)));
 	BehaviorTreeEditorTabFactories.RegisterFactory(MakeShareable(new FBehaviorTreeSearchSummoner(InBehaviorTreeEditor)));
+	BehaviorTreeEditorTabFactories.RegisterFactory(MakeShareable(new FBlackboardSummoner(InBehaviorTreeEditor)));
 
-	TabLayout = FTabManager::NewLayout( "Standalone_BehaviorTree_Layout" )
+	TabLayout = FTabManager::NewLayout( "Standalone_BehaviorTree_Layout_v1" )
 	->AddArea
 	(
 		FTabManager::NewPrimaryArea() ->SetOrientation(Orient_Vertical)
@@ -36,10 +37,21 @@ FBehaviorTreeEditorApplicationMode::FBehaviorTreeEditorApplicationMode(TSharedPt
 			)
 			->Split
 			(
-				FTabManager::NewStack()
+				FTabManager::NewSplitter() ->SetOrientation(Orient_Vertical)
 				->SetSizeCoefficient(0.3f)
-				->AddTab(FBehaviorTreeEditorTabs::GraphDetailsID, ETabState::OpenedTab)
-				->AddTab(FBehaviorTreeEditorTabs::SearchID, ETabState::ClosedTab)
+				->Split
+				(
+					FTabManager::NewStack()
+					->SetSizeCoefficient(0.6f)
+					->AddTab(FBehaviorTreeEditorTabs::GraphDetailsID, ETabState::OpenedTab)
+					->AddTab(FBehaviorTreeEditorTabs::SearchID, ETabState::ClosedTab)
+				)
+				->Split
+				(
+					FTabManager::NewStack()
+					->SetSizeCoefficient(0.4f)
+					->AddTab(FBehaviorTreeEditorTabs::BlackboardID, ETabState::OpenedTab)
+				)
 			)
 		)
 	);
@@ -86,9 +98,10 @@ FBlackboardEditorApplicationMode::FBlackboardEditorApplicationMode(TSharedPtr<cl
 {
 	BehaviorTreeEditor = InBehaviorTreeEditor;
 	
+	BlackboardTabFactories.RegisterFactory(MakeShareable(new FBlackboardEditorSummoner(InBehaviorTreeEditor)));
 	BlackboardTabFactories.RegisterFactory(MakeShareable(new FBlackboardDetailsSummoner(InBehaviorTreeEditor)));
 
-	TabLayout = FTabManager::NewLayout( "Standalone_BlackboardEditor_Layout" )
+	TabLayout = FTabManager::NewLayout( "Standalone_BlackboardEditor_Layout_v1" )
 	->AddArea
 	(
 		FTabManager::NewPrimaryArea() ->SetOrientation(Orient_Vertical)
@@ -101,7 +114,12 @@ FBlackboardEditorApplicationMode::FBlackboardEditorApplicationMode(TSharedPtr<cl
 		)
 		->Split
 		(
-			FTabManager::NewSplitter()
+			FTabManager::NewSplitter() ->SetOrientation(Orient_Horizontal)
+			->Split
+			(
+				FTabManager::NewStack()
+				->AddTab(FBehaviorTreeEditorTabs::BlackboardEditorID, ETabState::OpenedTab)
+			)
 			->Split
 			(
 				FTabManager::NewStack()
diff --git a/Engine/Source/Editor/BehaviorTreeEditor/Private/BehaviorTreeEditorTabFactories.cpp b/Engine/Source/Editor/BehaviorTreeEditor/Private/BehaviorTreeEditorTabFactories.cpp
--- a/Engine/Source/Editor/BehaviorTreeEditor/Private/BehaviorTreeEditorTabFactories.cpp
+++ b/Engine/Source/Editor/BehaviorTreeEditor/Private/BehaviorTreeEditorTabFactories.cpp
@@ -6,9 +6,56 @@
 #include "BehaviorTreeEditorTabs.h"
 #include "BehaviorTree/BehaviorTree.h"
 #include "BehaviorTree/BlackboardData.h"
+#include "SBehaviorTreeBlackboardView.h"
 
 #define LOCTEXT_NAMESPACE "BehaviorTreeEditorFactories"
 
+FBlackboardSummoner::FBlackboardSummoner(TSharedPtr<class FBehaviorTreeEditor> InBehaviorTreeEditorPtr)
+	: FWorkflowTabFactory(FBehaviorTreeEditorTabs::BlackboardID, InBehaviorTreeEditorPtr)
+	, BehaviorTreeEditorPtr(InBehaviorTreeEditorPtr)
+{
+	TabLabel = LOCTEXT("BlackboardLabel", "Blackboard");
+	TabIcon = FEditorStyle::GetBrush("Kismet.Tabs.Components");
+
+	bIsSingleton = true;
+
+	ViewMenuDescription = LOCTEXT("BlackboardView", "Blackboard");
+	ViewMenuTooltip = LOCTEXT("BlackboardView_ToolTip", "Show the blackboard view");
+}
+
+TSharedRef<SWidget> FBlackboardSummoner::CreateTabBody(const FWorkflowTabSpawnInfo& Info) const
+{
+	return BehaviorTreeEditorPtr.Pin()->SpawnBlackboardView();
+}
+
+FText FBlackboardSummoner::GetTabToolTipText(const FWorkflowTabSpawnInfo& Info) const
+{
+	return LOCTEXT("BlackboardTabTooltip", "The Blackboard view is for viewing and debugging blackboard key/value pairs.");
+}
+
+FBlackboardEditorSummoner::FBlackboardEditorSummoner(TSharedPtr<class FBehaviorTreeEditor> InBehaviorTreeEditorPtr)
+	: FWorkflowTabFactory(FBehaviorTreeEditorTabs::BlackboardEditorID, InBehaviorTreeEditorPtr)
+	, BehaviorTreeEditorPtr(InBehaviorTreeEditorPtr)
+{
+	TabLabel = LOCTEXT("BlackboardLabel", "Blackboard");
+	TabIcon = FEditorStyle::GetBrush("Kismet.Tabs.Components");
+
+	bIsSingleton = true;
+
+	ViewMenuDescription = LOCTEXT("BlackboardEditor", "Blackboard");
+	ViewMenuTooltip = LOCTEXT("BlackboardEditor_ToolTip", "Show the blackboard editor");
+}
+
+TSharedRef<SWidget> FBlackboardEditorSummoner::CreateTabBody(const FWorkflowTabSpawnInfo& Info) const
+{
+	return BehaviorTreeEditorPtr.Pin()->SpawnBlackboardEditor();
+}
+
+FText FBlackboardEditorSummoner::GetTabToolTipText(const FWorkflowTabSpawnInfo& Info) const
+{
+	return LOCTEXT("BlackboardEditorTabTooltip", "The Blackboard editor is for editing and debugging blackboard key/value pairs.");
+}
+
 FBlackboardDetailsSummoner::FBlackboardDetailsSummoner(TSharedPtr<class FBehaviorTreeEditor> InBehaviorTreeEditorPtr)
 	: FWorkflowTabFactory(FBehaviorTreeEditorTabs::BlackboardDetailsID, InBehaviorTreeEditorPtr)
 	, BehaviorTreeEditorPtr(InBehaviorTreeEditorPtr)
@@ -18,31 +65,19 @@ FBlackboardDetailsSummoner::FBlackboardDetailsSummoner(TSharedPtr<class FBehavio
 
 	bIsSingleton = true;
 
-	ViewMenuDescription = LOCTEXT("BlackboardDetailsView", "Blackboard Details");
-	ViewMenuTooltip = LOCTEXT("BlackboardDetailsView_ToolTip", "Show the blackboard details view");
+	ViewMenuDescription = LOCTEXT("BlackboardDetailsView", "Details");
+	ViewMenuTooltip = LOCTEXT("BlackboardDetailsView_ToolTip", "Show the details view");
 }
 
 TSharedRef<SWidget> FBlackboardDetailsSummoner::CreateTabBody(const FWorkflowTabSpawnInfo& Info) const
 {
-	const bool bIsUpdatable = false;
-	const bool bAllowFavorites = true;
-	const bool bIsLockable = false;
-	const bool bAllowSearch = true;
-	const bool bObjectsUseNameArea = false;
-	const bool bHideSelectionTip = true;
-
-	FPropertyEditorModule& PropertyEditorModule = FModuleManager::GetModuleChecked<FPropertyEditorModule>( "PropertyEditor" );
-	const FDetailsViewArgs DetailsViewArgs( bIsUpdatable, bIsLockable, bAllowSearch, bObjectsUseNameArea, bHideSelectionTip );
-	TSharedRef<IDetailsView> DetailsView = PropertyEditorModule.CreateDetailView( DetailsViewArgs );
-
-	DetailsView->SetObject(BehaviorTreeEditorPtr.Pin()->GetBlackboardData());
-
-	return DetailsView;
+	check(BehaviorTreeEditorPtr.IsValid());
+	return BehaviorTreeEditorPtr.Pin()->SpawnBlackboardDetails();
 }
 
 FText FBlackboardDetailsSummoner::GetTabToolTipText(const FWorkflowTabSpawnInfo& Info) const
 {
-	return LOCTEXT("BlackboardTabTooltip", "The Blackboard tab is for editing blackboard key/value pairs.");
+	return LOCTEXT("BlackboardDetailsTabTooltip", "The details tab is for editing blackboard entries.");
 }
 
 FBehaviorTreeDetailsSummoner::FBehaviorTreeDetailsSummoner(TSharedPtr<class FBehaviorTreeEditor> InBehaviorTreeEditorPtr)
@@ -54,8 +89,8 @@ FBehaviorTreeDetailsSummoner::FBehaviorTreeDetailsSummoner(TSharedPtr<class FBeh
 
 	bIsSingleton = true;
 
-	ViewMenuDescription = LOCTEXT("BehaviorTreeDetailsView", "Behavior Tree Details");
-	ViewMenuTooltip = LOCTEXT("BehaviorTreeDetailsView_ToolTip", "Show the behavior tree details view");
+	ViewMenuDescription = LOCTEXT("BehaviorTreeDetailsView", "Details");
+	ViewMenuTooltip = LOCTEXT("BehaviorTreeDetailsView_ToolTip", "Show the details view");
 }
 
 TSharedRef<SWidget> FBehaviorTreeDetailsSummoner::CreateTabBody(const FWorkflowTabSpawnInfo& Info) const
diff --git a/Engine/Source/Editor/BehaviorTreeEditor/Private/BehaviorTreeEditorTabFactories.h b/Engine/Source/Editor/BehaviorTreeEditor/Private/BehaviorTreeEditorTabFactories.h
--- a/Engine/Source/Editor/BehaviorTreeEditor/Private/BehaviorTreeEditorTabFactories.h
+++ b/Engine/Source/Editor/BehaviorTreeEditor/Private/BehaviorTreeEditorTabFactories.h
@@ -5,6 +5,30 @@
 #include "WorkflowOrientedApp/WorkflowTabFactory.h"
 #include "WorkflowOrientedApp/WorkflowUObjectDocuments.h"
 
+struct FBlackboardSummoner : public FWorkflowTabFactory
+{
+public:
+	FBlackboardSummoner(TSharedPtr<class FBehaviorTreeEditor> InBehaviorTreeEditorPtr);
+
+	virtual TSharedRef<SWidget> CreateTabBody(const FWorkflowTabSpawnInfo& Info) const override;
+	virtual FText GetTabToolTipText(const FWorkflowTabSpawnInfo& Info) const override;
+
+protected:
+	TWeakPtr<class FBehaviorTreeEditor> BehaviorTreeEditorPtr;
+};
+
+struct FBlackboardEditorSummoner : public FWorkflowTabFactory
+{
+public:
+	FBlackboardEditorSummoner(TSharedPtr<class FBehaviorTreeEditor> InBehaviorTreeEditorPtr);
+
+	virtual TSharedRef<SWidget> CreateTabBody(const FWorkflowTabSpawnInfo& Info) const override;
+	virtual FText GetTabToolTipText(const FWorkflowTabSpawnInfo& Info) const override;
+
+protected:
+	TWeakPtr<class FBehaviorTreeEditor> BehaviorTreeEditorPtr;
+};
+
 struct FBlackboardDetailsSummoner : public FWorkflowTabFactory
 {
 public:
diff --git a/Engine/Source/Editor/BehaviorTreeEditor/Private/BehaviorTreeEditorTabs.cpp b/Engine/Source/Editor/BehaviorTreeEditor/Private/BehaviorTreeEditorTabs.cpp
--- a/Engine/Source/Editor/BehaviorTreeEditor/Private/BehaviorTreeEditorTabs.cpp
+++ b/Engine/Source/Editor/BehaviorTreeEditor/Private/BehaviorTreeEditorTabs.cpp
@@ -6,6 +6,8 @@
 // Tab identifiers
 const FName FBehaviorTreeEditorTabs::GraphDetailsID(TEXT("BehaviorTreeEditor_Properties"));
 const FName FBehaviorTreeEditorTabs::SearchID(TEXT("BehaviorTreeEditor_Search"));
+const FName FBehaviorTreeEditorTabs::BlackboardID(TEXT("BehaviorTreeEditor_Blackboard"));
+const FName FBehaviorTreeEditorTabs::BlackboardEditorID(TEXT("BehaviorTreeEditor_BlackboardEditor"));
 const FName FBehaviorTreeEditorTabs::BlackboardDetailsID(TEXT("BehaviorTreeEditor_BlackboardProperties"));
 
 // Document tab identifiers
diff --git a/Engine/Source/Editor/BehaviorTreeEditor/Private/BehaviorTreeEditorTabs.h b/Engine/Source/Editor/BehaviorTreeEditor/Private/BehaviorTreeEditorTabs.h
--- a/Engine/Source/Editor/BehaviorTreeEditor/Private/BehaviorTreeEditorTabs.h
+++ b/Engine/Source/Editor/BehaviorTreeEditor/Private/BehaviorTreeEditorTabs.h
@@ -7,6 +7,8 @@ struct FBehaviorTreeEditorTabs
 	// Tab identifiers
 	static const FName GraphDetailsID;
 	static const FName SearchID;
+	static const FName BlackboardID;
+	static const FName BlackboardEditorID;
 	static const FName BlackboardDetailsID;
 
 	// Document tab identifiers
diff --git a/Engine/Source/Editor/BehaviorTreeEditor/Private/BehaviorTreeEditorToolbar.cpp b/Engine/Source/Editor/BehaviorTreeEditor/Private/BehaviorTreeEditorToolbar.cpp
--- a/Engine/Source/Editor/BehaviorTreeEditor/Private/BehaviorTreeEditorToolbar.cpp
+++ b/Engine/Source/Editor/BehaviorTreeEditor/Private/BehaviorTreeEditorToolbar.cpp
@@ -47,7 +47,6 @@ void FBehaviorTreeEditorToolbar::AddModesToolbar(TSharedPtr<FExtender> Extender)
 
 void FBehaviorTreeEditorToolbar::AddDebuggerToolbar(TSharedPtr<FExtender> Extender)
 {
-	check(BehaviorTreeEditor.IsValid());
 	TSharedPtr<FBehaviorTreeEditor> BehaviorTreeEditorPtr = BehaviorTreeEditor.Pin();
 
 	// setup toolbar
@@ -133,4 +132,4 @@ void FBehaviorTreeEditorToolbar::FillModesToolbar(FToolBarBuilder& ToolbarBuilde
 	BehaviorTreeEditorPtr->AddToolbarWidget(SNew(SSpacer).Size(FVector2D(4.0f, 1.0f)));
 }
 
-#undef LOCTEXT_NAMESPACE
\ No newline at end of file
+#undef LOCTEXT_NAMESPACE
diff --git a/Engine/Source/Editor/BehaviorTreeEditor/Private/DetailCustomizations/BlackboardDataDetails.cpp b/Engine/Source/Editor/BehaviorTreeEditor/Private/DetailCustomizations/BlackboardDataDetails.cpp
new file mode 100644
--- /dev/null
+++ b/Engine/Source/Editor/BehaviorTreeEditor/Private/DetailCustomizations/BlackboardDataDetails.cpp
@@ -0,0 +1,40 @@
+// Copyright 1998-2014 Epic Games, Inc. All Rights Reserved.
+
+#include "BehaviorTreeEditorPrivatePCH.h"
+#include "BlackboardDataDetails.h"
+
+
+TSharedRef<IDetailCustomization> FBlackboardDataDetails::MakeInstance(FOnGetSelectedBlackboardItemIndex InOnGetSelectedBlackboardItemIndex)
+{
+	return MakeShareable( new FBlackboardDataDetails(InOnGetSelectedBlackboardItemIndex) );
+}
+
+void FBlackboardDataDetails::CustomizeDetails( IDetailLayoutBuilder& DetailLayout )
+{
+	// First hide all keys
+	DetailLayout.HideProperty(TEXT("Keys"));
+	DetailLayout.HideProperty(TEXT("ParentKeys"));
+
+	// Now show only the currently selected key
+	bool bIsInherited = false;
+	int32 CurrentSelection = INDEX_NONE;
+	if(OnGetSelectedBlackboardItemIndex.IsBound())
+	{
+		CurrentSelection = OnGetSelectedBlackboardItemIndex.Execute(bIsInherited);
+	}
+
+	if(CurrentSelection >= 0)
+	{
+		TSharedPtr<IPropertyHandle> KeysHandle = bIsInherited ? DetailLayout.GetProperty(TEXT("ParentKeys")) : DetailLayout.GetProperty(TEXT("Keys"));
+		check(KeysHandle.IsValid());
+		uint32 NumChildKeys = 0;
+		KeysHandle->GetNumChildren(NumChildKeys);
+		if((uint32)CurrentSelection < NumChildKeys)
+		{
+			TSharedPtr<IPropertyHandle> KeyHandle = KeysHandle->GetChildHandle((uint32)CurrentSelection);
+
+			IDetailCategoryBuilder& DetailCategoryBuilder = DetailLayout.EditCategory("Key");
+			DetailCategoryBuilder.AddProperty(KeyHandle);
+		}	
+	}
+}
diff --git a/Engine/Source/Editor/BehaviorTreeEditor/Private/DetailCustomizations/BlackboardDataDetails.h b/Engine/Source/Editor/BehaviorTreeEditor/Private/DetailCustomizations/BlackboardDataDetails.h
new file mode 100644
--- /dev/null
+++ b/Engine/Source/Editor/BehaviorTreeEditor/Private/DetailCustomizations/BlackboardDataDetails.h
@@ -0,0 +1,27 @@
+// Copyright 1998-2014 Epic Games, Inc. All Rights Reserved.
+
+#pragma once
+
+#include "PropertyEditing.h"
+#include "PropertyCustomizationHelpers.h"
+
+/** Delegate used to retrieve current blackboard selection */
+DECLARE_DELEGATE_RetVal_OneParam(int32, FOnGetSelectedBlackboardItemIndex, bool& /* bIsInherited */);
+
+class FBlackboardDataDetails : public IDetailCustomization
+{
+public:
+	FBlackboardDataDetails(FOnGetSelectedBlackboardItemIndex InOnGetSelectedBlackboardItemIndex)
+		: OnGetSelectedBlackboardItemIndex(InOnGetSelectedBlackboardItemIndex)
+	{}
+
+	/** Makes a new instance of this detail layout class for a specific detail view requesting it */
+	static TSharedRef<IDetailCustomization> MakeInstance(FOnGetSelectedBlackboardItemIndex InOnGetSelectedBlackboardItemIndex);
+
+	/** IDetailCustomization interface */
+	virtual void CustomizeDetails( IDetailLayoutBuilder& DetailLayout ) override;
+
+private:
+	/** Delegate used to retrieve current blackboard selection */
+	FOnGetSelectedBlackboardItemIndex OnGetSelectedBlackboardItemIndex;
+};
diff --git a/Engine/Source/Editor/BehaviorTreeEditor/Private/DetailCustomizations/BlackboardDecoratorDetails.cpp b/Engine/Source/Editor/BehaviorTreeEditor/Private/DetailCustomizations/BlackboardDecoratorDetails.cpp
--- a/Engine/Source/Editor/BehaviorTreeEditor/Private/DetailCustomizations/BlackboardDecoratorDetails.cpp
+++ b/Engine/Source/Editor/BehaviorTreeEditor/Private/DetailCustomizations/BlackboardDecoratorDetails.cpp
@@ -124,8 +124,16 @@ void FBlackboardDecoratorDetails::OnKeyIDChanged()
 	if (Result == FPropertyAccess::Success)
 	{
 		const FBlackboardEntry* KeyEntry = Blackboard->GetKey(KeyID);
-		CachedKeyType = KeyEntry && KeyEntry->KeyType ? KeyEntry->KeyType->GetClass() : NULL;
-		CachedOperationType = KeyEntry->KeyType->GetTestOperation();
+		if(KeyEntry && KeyEntry->KeyType)
+		{
+			CachedKeyType = KeyEntry->KeyType->GetClass();
+			CachedOperationType = KeyEntry->KeyType->GetTestOperation();
+		}
+		else
+		{
+			CachedKeyType = nullptr;
+			CachedOperationType = 0;
+		}
 	}
 
 	// special handling of enum type: cache all names for combo box (display names)
diff --git a/Engine/Source/Editor/BehaviorTreeEditor/Private/SBehaviorTreeBlackboardEditor.cpp b/Engine/Source/Editor/BehaviorTreeEditor/Private/SBehaviorTreeBlackboardEditor.cpp
new file mode 100644
--- /dev/null
+++ b/Engine/Source/Editor/BehaviorTreeEditor/Private/SBehaviorTreeBlackboardEditor.cpp
@@ -0,0 +1,188 @@
+// Copyright 1998-2014 Epic Games, Inc. All Rights Reserved.
+
+#include "BehaviorTreeEditorPrivatePCH.h"
+#include "SBehaviorTreeBlackboardEditor.h"
+#include "BehaviorTree/BlackboardData.h"
+#include "BehaviorTree/Blackboard/BlackboardKeyAllTypes.h"
+
+#include "BehaviorTreeEditor.h"
+#include "ClassIconFinder.h"
+#include "ScopedTransaction.h"
+#include "GraphActionNode.h"
+#include "BehaviorTreeEditorCommands.h"
+#include "ClassViewerModule.h"
+#include "ClassViewerFilter.h"
+
+#define LOCTEXT_NAMESPACE "SBehaviorTreeBlackboardEditor"
+
+void SBehaviorTreeBlackboardEditor::Construct(const FArguments& InArgs, TSharedRef<FUICommandList> ToolkitCommands, UBlackboardData* InBlackboardData)
+{
+	OnEntrySelected = InArgs._OnEntrySelected;
+	OnGetDebugKeyValue = InArgs._OnGetDebugKeyValue;
+	OnIsDebuggerReady = InArgs._OnIsDebuggerReady;
+	OnIsDebuggerPaused = InArgs._OnIsDebuggerPaused;
+	OnGetDebugTimeStamp = InArgs._OnGetDebugTimeStamp;
+	OnGetDisplayCurrentState = InArgs._OnGetDisplayCurrentState;
+
+	ToolkitCommands->MapAction(
+		FBTBlackboardCommands::Get().DeleteEntry,
+		FExecuteAction::CreateSP(this, &SBehaviorTreeBlackboardEditor::HandleDeleteEntry),
+		FCanExecuteAction::CreateSP(this, &SBehaviorTreeBlackboardEditor::HasSelectedItems)
+		);
+
+	ToolkitCommands->MapAction(
+		FGenericCommands::Get().Rename,
+		FExecuteAction::CreateSP(this, &SBehaviorTreeBlackboardEditor::HandleRenameEntry),
+		FCanExecuteAction::CreateSP(this, &SBehaviorTreeBlackboardEditor::HasSelectedItems)
+		);
+
+	SBehaviorTreeBlackboardView::Construct(
+		SBehaviorTreeBlackboardView::FArguments()
+		.OnEntrySelected(InArgs._OnEntrySelected)
+		.OnGetDebugKeyValue(InArgs._OnGetDebugKeyValue)
+		.OnGetDisplayCurrentState(InArgs._OnGetDisplayCurrentState)
+		.OnIsDebuggerReady(InArgs._OnIsDebuggerReady)
+		.OnIsDebuggerPaused(InArgs._OnIsDebuggerPaused)
+		.OnGetDebugTimeStamp(InArgs._OnGetDebugTimeStamp)
+		.IsReadOnly(false),
+		InBlackboardData
+	);
+}
+
+void SBehaviorTreeBlackboardEditor::FillContextMenu(FMenuBuilder& MenuBuilder) const
+{
+	if(!IsDebuggerActive() && HasSelectedItems())
+	{
+		MenuBuilder.AddMenuEntry(FBTBlackboardCommands::Get().DeleteEntry);
+		MenuBuilder.AddMenuEntry(FGenericCommands::Get().Rename, NAME_None, LOCTEXT("Rename", "Rename"), LOCTEXT("Rename_Tooltip", "Renames this blackboard entry.") );
+	}
+}
+
+void SBehaviorTreeBlackboardEditor::FillToolbar(FToolBarBuilder& ToolbarBuilder) const
+{
+	ToolbarBuilder.AddComboButton(
+		FUIAction(
+			FExecuteAction(),
+			FCanExecuteAction::CreateSP(this, &SBehaviorTreeBlackboardEditor::CanCreateNewEntry)
+			), 
+		FOnGetContent::CreateSP(this, &SBehaviorTreeBlackboardEditor::HandleCreateNewEntryMenu),
+		LOCTEXT( "New_Label", "New Key" ),
+		LOCTEXT( "New_ToolTip", "Create a new blackboard entry" ),
+		FSlateIcon(FEditorStyle::GetStyleSetName(), "BTEditor.Blackboard.NewEntry")
+	);			
+}
+
+TSharedPtr<FExtender> SBehaviorTreeBlackboardEditor::GetToolbarExtender(TSharedRef<FUICommandList> ToolkitCommands) const
+{
+	TSharedPtr<FExtender> ToolbarExtender = MakeShareable(new FExtender);
+	ToolbarExtender->AddToolBarExtension("Debugging", EExtensionHook::Before, ToolkitCommands, FToolBarExtensionDelegate::CreateSP( this, &SBehaviorTreeBlackboardEditor::FillToolbar ));
+
+	return ToolbarExtender;
+}
+
+void SBehaviorTreeBlackboardEditor::HandleDeleteEntry()
+{
+	check(BlackboardData);
+
+	if(!IsDebuggerActive())
+	{
+		FBlackboardEntry* BlackboardEntry = GetSelectedEntry();
+		if(BlackboardEntry != nullptr)
+		{
+			const FScopedTransaction Transaction(LOCTEXT("BlackboardEntryDeleteTransaction", "Delete Blackboard Entry"));
+			BlackboardData->Modify();
+		
+			for(int32 ItemIndex = 0; ItemIndex < BlackboardData->Keys.Num(); ItemIndex++)
+			{
+				if(BlackboardEntry == &BlackboardData->Keys[ItemIndex])
+				{
+					BlackboardData->Keys.RemoveAt(ItemIndex);
+					break;
+				}
+			}
+
+			GraphActionMenu->RefreshAllActions(true);
+		}
+	}
+}
+
+void SBehaviorTreeBlackboardEditor::HandleRenameEntry()
+{
+	if(!IsDebuggerActive())
+	{
+		GraphActionMenu->OnRequestRenameOnActionNode();
+	}
+}
+
+class FBlackboardEntryClassFilter : public IClassViewerFilter
+{
+public:
+
+	virtual bool IsClassAllowed(const FClassViewerInitializationOptions& InInitOptions, const UClass* InClass, TSharedRef< FClassViewerFilterFuncs > InFilterFuncs ) override
+	{
+		if(InClass != nullptr)
+		{
+			return !InClass->HasAnyClassFlags(CLASS_Abstract) && InClass->IsChildOf(UBlackboardKeyType::StaticClass());
+		}
+		return false;
+	}
+
+	virtual bool IsUnloadedClassAllowed(const FClassViewerInitializationOptions& InInitOptions, const TSharedRef< const IUnloadedBlueprintData > InUnloadedClassData, TSharedRef< FClassViewerFilterFuncs > InFilterFuncs) override
+	{
+		return InUnloadedClassData->IsChildOf(UBlackboardKeyType::StaticClass());
+	}
+};
+
+
+TSharedRef<SWidget> SBehaviorTreeBlackboardEditor::HandleCreateNewEntryMenu() const
+{
+	FClassViewerInitializationOptions Options;
+	Options.bShowUnloadedBlueprints = true;
+	Options.bShowDisplayNames = true;
+	Options.ClassFilter = MakeShareable( new FBlackboardEntryClassFilter );
+
+	FOnClassPicked OnPicked( FOnClassPicked::CreateRaw( this, &SBehaviorTreeBlackboardEditor::HandleKeyClassPicked ) );
+
+	return FModuleManager::LoadModuleChecked<FClassViewerModule>("ClassViewer").CreateClassViewer(Options, OnPicked);
+}
+
+void SBehaviorTreeBlackboardEditor::HandleKeyClassPicked(UClass* InClass)
+{
+	check(BlackboardData);
+
+	FSlateApplication::Get().DismissAllMenus();
+
+	check(InClass);
+	check(InClass->IsChildOf(UBlackboardKeyType::StaticClass()));
+
+	const FScopedTransaction Transaction(LOCTEXT("BlackboardEntryAddTransaction", "Add Blackboard Entry"));
+	BlackboardData->Modify();
+
+	// create a name for this new key
+	FString NewKeyName = InClass->GetDisplayNameText().ToString();
+	NewKeyName = NewKeyName.Replace(TEXT(" "), TEXT(""));
+	NewKeyName += TEXT("Key");
+
+	FBlackboardEntry Entry;
+	Entry.EntryName = FName(*NewKeyName);
+	Entry.KeyType = ConstructObject<UBlackboardKeyType>(InClass, BlackboardData);		
+
+	BlackboardData->Keys.Add(Entry);
+
+	GraphActionMenu->RefreshAllActions(true);
+
+	GraphActionMenu->SelectItemByName(Entry.EntryName, ESelectInfo::OnMouseClick);
+	GraphActionMenu->OnRequestRenameOnActionNode();
+}
+
+bool SBehaviorTreeBlackboardEditor::CanCreateNewEntry() const
+{
+	if(OnIsDebuggerReady.IsBound())
+	{
+		return !OnIsDebuggerReady.Execute();
+	}
+
+	return true;
+}
+
+#undef LOCTEXT_NAMESPACE
\ No newline at end of file
diff --git a/Engine/Source/Editor/BehaviorTreeEditor/Private/SBehaviorTreeBlackboardEditor.h b/Engine/Source/Editor/BehaviorTreeEditor/Private/SBehaviorTreeBlackboardEditor.h
new file mode 100644
--- /dev/null
+++ b/Engine/Source/Editor/BehaviorTreeEditor/Private/SBehaviorTreeBlackboardEditor.h
@@ -0,0 +1,50 @@
+// Copyright 1998-2014 Epic Games, Inc. All Rights Reserved.
+
+#pragma once
+
+#include "SGraphPalette.h"
+#include "EdGraph/EdGraphSchema.h"
+#include "SBehaviorTreeBlackboardView.h"
+
+/** Displays and edits blackboard entries */
+class SBehaviorTreeBlackboardEditor : public SBehaviorTreeBlackboardView
+{
+public:
+	SLATE_BEGIN_ARGS( SBehaviorTreeBlackboardEditor ) {}
+
+		SLATE_EVENT(FOnEntrySelected, OnEntrySelected)
+		SLATE_EVENT(FOnGetDebugKeyValue, OnGetDebugKeyValue)
+		SLATE_EVENT(FOnGetDisplayCurrentState, OnGetDisplayCurrentState)
+		SLATE_EVENT(FOnIsDebuggerReady, OnIsDebuggerReady)
+		SLATE_EVENT(FOnIsDebuggerPaused, OnIsDebuggerPaused)
+		SLATE_EVENT(FOnGetDebugTimeStamp, OnGetDebugTimeStamp)
+
+	SLATE_END_ARGS()
+
+	void Construct(const FArguments& InArgs, TSharedRef<FUICommandList> ToolkitCommands, UBlackboardData* InBlackboardData);
+
+private:
+	/** Fill the context menu with edit options */
+	virtual void FillContextMenu(FMenuBuilder& MenuBuilder) const override;
+
+	/** Fill the toolbar with edit options */
+	void FillToolbar(FToolBarBuilder& ToolbarBuilder) const;
+
+	/** Extend the toolbar */
+	virtual TSharedPtr<FExtender> GetToolbarExtender(TSharedRef<FUICommandList> ToolkitCommands) const override;
+
+	/** Handle deleting the currently selected entry */
+	void HandleDeleteEntry();
+
+	/** Handle renaming the currently selected entry */
+	void HandleRenameEntry();
+
+	/** Create the menu used to create a new blackboard entry */
+	TSharedRef<SWidget> HandleCreateNewEntryMenu() const;
+
+	/** Create a new blackboard entry when a class is picked */
+	void HandleKeyClassPicked(UClass* InClass);
+
+	/** Delegate handler that disallows creating a new entry while the debugger is running */
+	bool CanCreateNewEntry() const;
+};
\ No newline at end of file
diff --git a/Engine/Source/Editor/BehaviorTreeEditor/Private/SBehaviorTreeBlackboardView.cpp b/Engine/Source/Editor/BehaviorTreeEditor/Private/SBehaviorTreeBlackboardView.cpp
new file mode 100644
--- /dev/null
+++ b/Engine/Source/Editor/BehaviorTreeEditor/Private/SBehaviorTreeBlackboardView.cpp
@@ -0,0 +1,539 @@
+// Copyright 1998-2014 Epic Games, Inc. All Rights Reserved.
+
+#include "BehaviorTreeEditorPrivatePCH.h"
+#include "SBehaviorTreeBlackboardView.h"
+#include "BehaviorTree/BlackboardData.h"
+#include "BehaviorTree/Blackboard/BlackboardKeyAllTypes.h"
+#include "SGraphActionMenu.h"
+#include "BehaviorTreeEditor.h"
+#include "ClassIconFinder.h"
+#include "ScopedTransaction.h"
+#include "GraphActionNode.h"
+#include "BehaviorTreeEditorCommands.h"
+#include "ClassViewerModule.h"
+#include "ClassViewerFilter.h"
+
+#define LOCTEXT_NAMESPACE "SBehaviorTreeBlackboardView"
+
+namespace EBlackboardSectionTitles
+{
+	enum Type
+	{
+		InheritedKeys = 1,
+		Keys,
+	};
+}
+
+
+class FEdGraphSchemaAction_BlackboardEntry : public FEdGraphSchemaAction_Dummy
+{
+public:
+	static FString StaticGetTypeId() { static FString Type = TEXT("FEdGraphSchemaAction_BlackboardEntry"); return Type; }
+	virtual FString GetTypeId() const override { return StaticGetTypeId(); }
+
+	FEdGraphSchemaAction_BlackboardEntry( UBlackboardData* InBlackboardData, FBlackboardEntry& InKey, bool bInIsInherited )
+		: FEdGraphSchemaAction_Dummy()
+		, BlackboardData(InBlackboardData)
+		, Key(InKey)
+		, bIsInherited(bInIsInherited)
+	{
+		check(BlackboardData);
+		Update();
+	}
+
+	void Update()
+	{
+		MenuDescription = FText::FromName(Key.EntryName);
+		TooltipDescription = FText::Format(LOCTEXT("BlackboardEntryFormat", "{0} '{1}'"), Key.KeyType->GetClass()->GetDisplayNameText(), FText::FromName(Key.EntryName)).ToString();
+		SectionID = bIsInherited ? EBlackboardSectionTitles::InheritedKeys : EBlackboardSectionTitles::Keys;	
+	}
+
+	UBlackboardData* BlackboardData;
+
+	FBlackboardEntry& Key;
+
+	bool bIsInherited;
+};
+
+
+class SBehaviorTreeBlackboardItem : public SGraphPaletteItem
+{
+	SLATE_BEGIN_ARGS( SBehaviorTreeBlackboardItem ) {}
+
+		SLATE_EVENT(FOnGetDebugKeyValue, OnGetDebugKeyValue)
+		SLATE_EVENT(FOnGetDisplayCurrentState, OnGetDisplayCurrentState)
+		SLATE_EVENT(FOnIsDebuggerReady, OnIsDebuggerReady)
+
+	SLATE_END_ARGS()
+
+	void Construct(const FArguments& InArgs, FCreateWidgetForActionData* const InCreateData)
+	{
+		OnGetDebugKeyValue = InArgs._OnGetDebugKeyValue;
+		OnIsDebuggerReady = InArgs._OnIsDebuggerReady;
+		OnGetDisplayCurrentState = InArgs._OnGetDisplayCurrentState;
+
+		const FSlateFontInfo NameFont = FSlateFontInfo( FPaths::EngineContentDir() / TEXT("Slate/Fonts/Roboto-Regular.ttf"), 10);
+
+		check(InCreateData);
+		check(InCreateData->Action.IsValid());
+
+		TSharedPtr<FEdGraphSchemaAction> GraphAction = InCreateData->Action;
+		check(GraphAction->GetTypeId() == FEdGraphSchemaAction_BlackboardEntry::StaticGetTypeId());
+		TSharedPtr<FEdGraphSchemaAction_BlackboardEntry> BlackboardEntryAction = StaticCastSharedPtr<FEdGraphSchemaAction_BlackboardEntry>(GraphAction);
+
+		ActionPtr = InCreateData->Action;
+		
+		FSlateBrush const* IconBrush   = FEditorStyle::GetBrush(TEXT("NoBrush"));
+		GetPaletteItemIcon(GraphAction, IconBrush);
+
+		TSharedRef<SWidget> IconWidget = CreateIconWidget( GraphAction->TooltipDescription, IconBrush, IconBrush->TintColor );
+		TSharedRef<SWidget> NameSlotWidget = CreateTextSlotWidget( NameFont, InCreateData, BlackboardEntryAction->bIsInherited );
+		TSharedRef<SWidget> DebugSlotWidget = CreateDebugSlotWidget( NameFont );
+
+		// Create the actual widget
+		this->ChildSlot
+		[
+			SNew(SHorizontalBox)
+			// Icon slot
+			+SHorizontalBox::Slot()
+			.AutoWidth()
+			.VAlign(VAlign_Center)
+			[
+				IconWidget
+			]
+			// Name slot
+			+SHorizontalBox::Slot()
+			.AutoWidth()
+			.VAlign(VAlign_Center)
+			.Padding(3,0)
+			[
+				NameSlotWidget
+			]
+			// Debug info slot
+			+SHorizontalBox::Slot()
+			.FillWidth(1.0f)
+			.HAlign(HAlign_Right)
+			.VAlign(VAlign_Center)
+			.Padding(3,0)
+			[
+				DebugSlotWidget
+			]
+		];
+	}
+
+private:
+	void GetPaletteItemIcon(TSharedPtr<FEdGraphSchemaAction> InGraphAction, FSlateBrush const*& OutIconBrush)
+	{
+		check(InGraphAction.IsValid());
+		check(InGraphAction->GetTypeId() == FEdGraphSchemaAction_BlackboardEntry::StaticGetTypeId());
+		TSharedPtr<FEdGraphSchemaAction_BlackboardEntry> BlackboardEntryAction = StaticCastSharedPtr<FEdGraphSchemaAction_BlackboardEntry>(InGraphAction);
+		check(BlackboardEntryAction->Key.KeyType);
+
+		OutIconBrush = FClassIconFinder::FindIconForClass(BlackboardEntryAction->Key.KeyType->GetClass());
+	}
+
+	virtual TSharedRef<SWidget> CreateTextSlotWidget( const FSlateFontInfo& NameFont, FCreateWidgetForActionData* const InCreateData, bool bInIsReadOnly ) override
+	{
+		check(InCreateData);
+
+		TSharedPtr< SWidget > DisplayWidget;
+
+		// Copy the mouse delegate binding if we want it
+		if( InCreateData->bHandleMouseButtonDown )
+		{
+			MouseButtonDownDelegate = InCreateData->MouseButtonDownDelegate;
+		}
+
+		// If the creation data says read only, then it must be read only
+		bIsReadOnly = InCreateData->bIsReadOnly || bInIsReadOnly;
+
+		InlineRenameWidget =
+			SAssignNew(DisplayWidget, SInlineEditableTextBlock)
+			.Text(this, &SBehaviorTreeBlackboardItem::GetDisplayText)
+			.Font(NameFont)
+			.HighlightText(InCreateData->HighlightText)
+			.ToolTipText(this, &SBehaviorTreeBlackboardItem::GetItemTooltip)
+			.OnTextCommitted(this, &SBehaviorTreeBlackboardItem::OnNameTextCommitted)
+			.OnVerifyTextChanged(this, &SBehaviorTreeBlackboardItem::OnNameTextVerifyChanged)
+			.IsSelected( InCreateData->IsRowSelectedDelegate )
+			.IsReadOnly(this, &SBehaviorTreeBlackboardItem::IsReadOnly);
+
+		if(!bIsReadOnly)
+		{
+			InCreateData->OnRenameRequest->BindSP( InlineRenameWidget.Get(), &SInlineEditableTextBlock::EnterEditingMode );
+		}
+
+		return DisplayWidget.ToSharedRef();
+	}
+
+	virtual FString GetItemTooltip() const override
+	{
+		return ActionPtr.Pin()->TooltipDescription;
+	}
+
+	virtual void OnNameTextCommitted(const FText& NewText, ETextCommit::Type InTextCommit) override
+	{
+		check(ActionPtr.Pin()->GetTypeId() == FEdGraphSchemaAction_BlackboardEntry::StaticGetTypeId());
+		TSharedPtr<FEdGraphSchemaAction_BlackboardEntry> BlackboardEntryAction = StaticCastSharedPtr<FEdGraphSchemaAction_BlackboardEntry>(ActionPtr.Pin());
+
+		FName NewName = FName(*NewText.ToString());
+		if(NewName != BlackboardEntryAction->Key.EntryName)
+		{
+			const FScopedTransaction Transaction(LOCTEXT("BlackboardEntryRenameTransaction", "Rename Blackboard Entry"));
+			BlackboardEntryAction->BlackboardData->Modify();
+			BlackboardEntryAction->Key.EntryName = FName(*NewText.ToString());
+
+			BlackboardEntryAction->Update();
+		}
+	}
+
+	/** Create widget for displaying debug information about this blackboard entry */
+	TSharedRef<SWidget> CreateDebugSlotWidget(const FSlateFontInfo& InFontInfo)
+	{
+		check(ActionPtr.Pin()->GetTypeId() == FEdGraphSchemaAction_BlackboardEntry::StaticGetTypeId());
+		TSharedPtr<FEdGraphSchemaAction_BlackboardEntry> BlackboardEntryAction = StaticCastSharedPtr<FEdGraphSchemaAction_BlackboardEntry>(ActionPtr.Pin());
+
+		return SNew(STextBlock)
+			.Text(this, &SBehaviorTreeBlackboardItem::GetDebugTextValue, BlackboardEntryAction)
+			.Visibility(this, &SBehaviorTreeBlackboardItem::GetDebugTextVisibility);
+	}
+
+	FText GetDebugTextValue(TSharedPtr<FEdGraphSchemaAction_BlackboardEntry> BlackboardEntry) const
+	{
+		check(BlackboardEntry.IsValid());
+		if(OnGetDebugKeyValue.IsBound() && OnGetDisplayCurrentState.IsBound())
+		{
+			return OnGetDebugKeyValue.Execute(BlackboardEntry->Key.EntryName, OnGetDisplayCurrentState.Execute());
+		}
+		
+		return FText();
+	}
+
+	EVisibility GetDebugTextVisibility() const
+	{
+		if(OnIsDebuggerReady.IsBound())
+		{
+			return OnIsDebuggerReady.Execute() ? EVisibility::Visible : EVisibility::Collapsed;
+		}
+
+		return EVisibility::Collapsed;
+	}
+
+	bool IsReadOnly() const
+	{
+		if(OnIsDebuggerReady.IsBound())
+		{
+			return bIsReadOnly || OnIsDebuggerReady.Execute();
+		}
+
+		return bIsReadOnly;
+	}
+
+private:
+	/** Delegate used to retrieve debug data to display */
+	FOnGetDebugKeyValue OnGetDebugKeyValue;
+
+	/** Delegate used to determine whether the BT debugger is active */
+	FOnIsDebuggerReady OnIsDebuggerReady;
+
+	/** Delegate used to determine whether the BT debugger displaying the current state */
+	FOnGetDisplayCurrentState OnGetDisplayCurrentState;
+
+	/** Read-only flag */
+	bool bIsReadOnly;
+};
+
+
+void SBehaviorTreeBlackboardView::AddReferencedObjects( FReferenceCollector& Collector )
+{
+	if(BlackboardData != nullptr)
+	{
+		Collector.AddReferencedObject(BlackboardData);
+	}
+}
+
+void SBehaviorTreeBlackboardView::Construct(const FArguments& InArgs, UBlackboardData* InBlackboardData)
+{
+	OnEntrySelected = InArgs._OnEntrySelected;
+	OnGetDebugKeyValue = InArgs._OnGetDebugKeyValue;
+	OnIsDebuggerReady = InArgs._OnIsDebuggerReady;
+	OnIsDebuggerPaused = InArgs._OnIsDebuggerPaused;
+	OnGetDebugTimeStamp = InArgs._OnGetDebugTimeStamp;
+	OnGetDisplayCurrentState = InArgs._OnGetDisplayCurrentState;
+
+	BlackboardData = InBlackboardData;
+
+	bShowCurrentState = OnGetDisplayCurrentState.IsBound() ? OnGetDisplayCurrentState.Execute() : true;
+
+	TSharedRef<FUICommandList> ToolkitCommands = MakeShareable(new FUICommandList);
+
+	ToolkitCommands->MapAction(
+		FBTDebuggerCommands::Get().CurrentValues,
+		FUIAction(
+			FExecuteAction::CreateSP(this, &SBehaviorTreeBlackboardView::HandleUseCurrentValues),
+			FCanExecuteAction::CreateSP(this, &SBehaviorTreeBlackboardView::IsDebuggerPaused),
+			FIsActionChecked::CreateSP(this, &SBehaviorTreeBlackboardView::IsUsingCurrentValues),
+			FIsActionButtonVisible::CreateSP(this, &SBehaviorTreeBlackboardView::IsDebuggerActive)
+			)
+		);
+	ToolkitCommands->MapAction(
+		FBTDebuggerCommands::Get().SavedValues, 
+		FUIAction(
+			FExecuteAction::CreateSP(this, &SBehaviorTreeBlackboardView::HandleUseSavedValues),
+			FCanExecuteAction::CreateSP(this, &SBehaviorTreeBlackboardView::IsDebuggerPaused),
+			FIsActionChecked::CreateSP(this, &SBehaviorTreeBlackboardView::IsUsingSavedValues),
+			FIsActionButtonVisible::CreateSP(this, &SBehaviorTreeBlackboardView::IsDebuggerActive)
+			)
+		);
+
+	// build debug toolbar
+	FToolBarBuilder ToolbarBuilder(ToolkitCommands, FMultiBoxCustomization::None, GetToolbarExtender(ToolkitCommands));
+	
+	ToolbarBuilder.BeginSection(TEXT("Debugging"));
+	{
+		ToolbarBuilder.AddToolBarButton(FBTDebuggerCommands::Get().CurrentValues);
+		ToolbarBuilder.AddToolBarButton(FBTDebuggerCommands::Get().SavedValues);
+	}
+	ToolbarBuilder.EndSection();
+
+	ChildSlot
+	[
+		SNew(SBorder)
+		.Padding(4.0f)
+		.BorderImage(FEditorStyle::GetBrush("ToolPanel.GroupBorder"))
+		[
+			SNew(SVerticalBox)
+			+SVerticalBox::Slot()
+			.AutoHeight()
+			.Padding(0.0f, 0.0f, 0.0f, 4.0f)
+			[
+				SNew(SHorizontalBox)
+				+SHorizontalBox::Slot()
+				.AutoWidth()
+				[
+					ToolbarBuilder.MakeWidget()
+				]
+				+SHorizontalBox::Slot()
+				.AutoWidth()
+				.VAlign(VAlign_Center)
+				[
+					SNew(STextBlock)
+					.Text(this, &SBehaviorTreeBlackboardView::GetDebugTimeStampText)
+					.Visibility(this, &SBehaviorTreeBlackboardView::GetDebuggerToolbarVisibility)
+				]
+			]
+			+SVerticalBox::Slot()
+			.FillHeight(1.0f)
+			[
+				SAssignNew(GraphActionMenu, SGraphActionMenu, InArgs._IsReadOnly)
+				.OnCreateWidgetForAction(this, &SBehaviorTreeBlackboardView::HandleCreateWidgetForAction)
+				.OnCollectAllActions(this, &SBehaviorTreeBlackboardView::HandleCollectAllActions)
+				.OnGetSectionTitle(this, &SBehaviorTreeBlackboardView::HandleGetSectionTitle)
+				.OnActionSelected(this, &SBehaviorTreeBlackboardView::HandleActionSelected)
+				.OnContextMenuOpening(this, &SBehaviorTreeBlackboardView::HandleContextMenuOpening, ToolkitCommands)
+				.OnActionMatchesName(this, &SBehaviorTreeBlackboardView::HandleActionMatchesName)
+				.AlphaSortItems(false)
+			]
+		]
+	];
+}
+
+TSharedRef<SWidget> SBehaviorTreeBlackboardView::HandleCreateWidgetForAction(FCreateWidgetForActionData* const InCreateData)
+{
+	return SNew(SBehaviorTreeBlackboardItem, InCreateData)
+		.OnIsDebuggerReady(OnIsDebuggerReady)
+		.OnGetDebugKeyValue(OnGetDebugKeyValue)
+		.OnGetDisplayCurrentState(this, &SBehaviorTreeBlackboardView::IsUsingCurrentValues);
+}
+
+void SBehaviorTreeBlackboardView::HandleCollectAllActions( FGraphActionListBuilderBase& GraphActionListBuilder )
+{
+	if(BlackboardData != nullptr)
+	{
+		for(auto& ParentKey : BlackboardData->ParentKeys)
+		{
+			GraphActionListBuilder.AddAction( MakeShareable(new FEdGraphSchemaAction_BlackboardEntry(BlackboardData, ParentKey, true)) );
+		}
+
+		for(auto& Key : BlackboardData->Keys)
+		{
+			GraphActionListBuilder.AddAction( MakeShareable(new FEdGraphSchemaAction_BlackboardEntry(BlackboardData, Key, false)) );
+		}
+	}
+}
+
+FText SBehaviorTreeBlackboardView::HandleGetSectionTitle(int32 SectionID) const
+{
+	switch(SectionID)
+	{
+	case EBlackboardSectionTitles::InheritedKeys:
+		return LOCTEXT("InheritedKeysSectionLabel", "Inherited Keys");
+	case EBlackboardSectionTitles::Keys:
+		return LOCTEXT("KeysSectionLabel", "Keys");
+	}
+
+	return FText();
+}
+
+void SBehaviorTreeBlackboardView::HandleActionSelected(const TArray< TSharedPtr<FEdGraphSchemaAction> >& SelectedActions) const
+{
+	if(SelectedActions.Num() > 0)
+	{
+		check(SelectedActions[0]->GetTypeId() == FEdGraphSchemaAction_BlackboardEntry::StaticGetTypeId());
+		TSharedPtr<FEdGraphSchemaAction_BlackboardEntry> BlackboardEntry = StaticCastSharedPtr<FEdGraphSchemaAction_BlackboardEntry>(SelectedActions[0]);
+		OnEntrySelected.ExecuteIfBound(&BlackboardEntry->Key, BlackboardEntry->bIsInherited);
+	}
+}
+
+TSharedPtr<FEdGraphSchemaAction_BlackboardEntry> SBehaviorTreeBlackboardView::GetSelectedEntryInternal() const
+{
+	TArray< TSharedPtr<FEdGraphSchemaAction> > SelectedActions;
+	GraphActionMenu->GetSelectedActions(SelectedActions);
+
+	if(SelectedActions.Num() > 0)
+	{
+		check(SelectedActions[0]->GetTypeId() == FEdGraphSchemaAction_BlackboardEntry::StaticGetTypeId());
+		return StaticCastSharedPtr<FEdGraphSchemaAction_BlackboardEntry>(SelectedActions[0]);
+	}
+
+	return TSharedPtr<FEdGraphSchemaAction_BlackboardEntry>();
+}
+
+int32 SBehaviorTreeBlackboardView::GetSelectedEntryIndex(bool& bOutIsInherited) const
+{
+	TSharedPtr<FEdGraphSchemaAction_BlackboardEntry> Entry = GetSelectedEntryInternal();
+	if(Entry.IsValid())
+	{
+		bOutIsInherited = Entry->bIsInherited;
+		FBlackboardEntry* BlackboardEntry = &Entry->Key;
+
+		// check to see what entry index we are using
+		TArray<FBlackboardEntry>& EntryArray = bOutIsInherited ? BlackboardData->ParentKeys : BlackboardData->Keys;
+		for(int32 Index = 0; Index < EntryArray.Num(); Index++)
+		{
+			if(BlackboardEntry == &EntryArray[Index])
+			{
+				return Index;
+			}
+		}
+	}
+
+	return INDEX_NONE;
+}
+
+FBlackboardEntry* SBehaviorTreeBlackboardView::GetSelectedEntry() const
+{
+	TSharedPtr<FEdGraphSchemaAction_BlackboardEntry> Entry = GetSelectedEntryInternal();
+	if(Entry.IsValid())
+	{
+		return &Entry->Key;
+	}
+
+	return nullptr;
+}
+
+void SBehaviorTreeBlackboardView::SetObject(UBlackboardData* InBlackboardData)
+{
+	BlackboardData = InBlackboardData;
+	GraphActionMenu->RefreshAllActions(true);
+}
+
+TSharedPtr<SWidget> SBehaviorTreeBlackboardView::HandleContextMenuOpening(TSharedRef<FUICommandList> ToolkitCommands) const
+{
+	FMenuBuilder MenuBuilder(/* bInShouldCloseWindowAfterMenuSelection =*/true, ToolkitCommands);
+
+	FillContextMenu(MenuBuilder);
+
+	return MenuBuilder.MakeWidget();
+}
+
+void SBehaviorTreeBlackboardView::FillContextMenu(FMenuBuilder& MenuBuilder) const
+{
+
+}
+
+TSharedPtr<FExtender> SBehaviorTreeBlackboardView::GetToolbarExtender(TSharedRef<FUICommandList> ToolkitCommands) const
+{
+	return TSharedPtr<FExtender>();
+}
+
+void SBehaviorTreeBlackboardView::HandleUseCurrentValues()
+{
+	bShowCurrentState = true;
+}
+
+void SBehaviorTreeBlackboardView::HandleUseSavedValues()
+{
+	bShowCurrentState = false;
+}
+
+FText SBehaviorTreeBlackboardView::GetDebugTimeStampText() const
+{
+	FText TimeStampText;
+
+	if(OnGetDebugTimeStamp.IsBound())
+	{
+		TimeStampText = FText::Format(LOCTEXT("ToolbarTimeStamp", "Time Stamp: {0}"), FText::AsNumber(OnGetDebugTimeStamp.Execute(IsUsingCurrentValues())));
+	}
+
+	return TimeStampText;
+}
+
+EVisibility SBehaviorTreeBlackboardView::GetDebuggerToolbarVisibility() const
+{
+	if(OnIsDebuggerReady.IsBound())
+	{
+		return OnIsDebuggerReady.Execute() ? EVisibility::Visible : EVisibility::Collapsed;
+	}
+
+	return EVisibility::Collapsed;
+}
+
+bool SBehaviorTreeBlackboardView::IsUsingCurrentValues() const
+{
+	if(OnGetDisplayCurrentState.IsBound())
+	{
+		return OnGetDisplayCurrentState.Execute() || bShowCurrentState;
+	}
+
+	return bShowCurrentState;
+}
+
+bool SBehaviorTreeBlackboardView::IsUsingSavedValues() const
+{
+	return !IsUsingCurrentValues();
+}
+
+bool SBehaviorTreeBlackboardView::HasSelectedItems() const
+{
+	return GetSelectedEntry() != nullptr;
+}
+
+bool SBehaviorTreeBlackboardView::IsDebuggerActive() const
+{
+	if(OnIsDebuggerReady.IsBound())
+	{
+		return OnIsDebuggerReady.Execute();
+	}
+
+	return true;
+}
+
+bool SBehaviorTreeBlackboardView::IsDebuggerPaused() const
+{
+	if(OnIsDebuggerPaused.IsBound())
+	{
+		return OnIsDebuggerPaused.Execute();
+	}
+
+	return true;
+}
+
+bool SBehaviorTreeBlackboardView::HandleActionMatchesName(FEdGraphSchemaAction* InAction, const FName& InName) const
+{
+	check(InAction->GetTypeId() == FEdGraphSchemaAction_BlackboardEntry::StaticGetTypeId());
+	FEdGraphSchemaAction_BlackboardEntry* BlackboardEntryAction = static_cast<FEdGraphSchemaAction_BlackboardEntry*>(InAction);
+	return BlackboardEntryAction->Key.EntryName == InName;
+}
+
+#undef LOCTEXT_NAMESPACE
\ No newline at end of file
diff --git a/Engine/Source/Editor/BehaviorTreeEditor/Private/SBehaviorTreeBlackboardView.h b/Engine/Source/Editor/BehaviorTreeEditor/Private/SBehaviorTreeBlackboardView.h
new file mode 100644
--- /dev/null
+++ b/Engine/Source/Editor/BehaviorTreeEditor/Private/SBehaviorTreeBlackboardView.h
@@ -0,0 +1,154 @@
+// Copyright 1998-2014 Epic Games, Inc. All Rights Reserved.
+
+#pragma once
+
+#include "SGraphPalette.h"
+#include "EdGraph/EdGraphSchema.h"
+
+class UBlackboardData;
+struct FBlackboardEntry;
+class FEdGraphSchemaAction_BlackboardEntry;
+
+/** Delegate executed when an entry is selected */
+DECLARE_DELEGATE_TwoParams( FOnEntrySelected, const FBlackboardEntry*, bool );
+
+/** Delegate used to retrieve debug data to display */
+DECLARE_DELEGATE_RetVal_TwoParams(FText, FOnGetDebugKeyValue, const FName& /* InKeyName */, bool /* bUseCurrentState */);
+
+/** Delegate used to determine whether the BT debugger is active */
+DECLARE_DELEGATE_RetVal(bool, FOnIsDebuggerReady);
+
+/** Delegate used to determine whether the BT debugger is paused */
+DECLARE_DELEGATE_RetVal(bool, FOnIsDebuggerPaused);
+
+/** Delegate used to determine whether the BT debugger displaying the current state */
+DECLARE_DELEGATE_RetVal(bool, FOnGetDisplayCurrentState);
+
+/** Delegate used to get the debugger's current timestamp */
+DECLARE_DELEGATE_RetVal_OneParam(float, FOnGetDebugTimeStamp, bool /* bUseCurrentState */);
+
+
+/** Displays blackboard entries */
+class SBehaviorTreeBlackboardView : public SCompoundWidget, public FGCObject
+{
+public:
+	SLATE_BEGIN_ARGS( SBehaviorTreeBlackboardView ) 
+	{
+		_IsReadOnly = true;
+	}
+
+		SLATE_EVENT(FOnEntrySelected, OnEntrySelected)
+		SLATE_EVENT(FOnGetDebugKeyValue, OnGetDebugKeyValue)
+		SLATE_EVENT(FOnGetDisplayCurrentState, OnGetDisplayCurrentState)
+		SLATE_EVENT(FOnIsDebuggerReady, OnIsDebuggerReady)
+		SLATE_EVENT(FOnIsDebuggerPaused, OnIsDebuggerPaused)
+		SLATE_EVENT(FOnGetDebugTimeStamp, OnGetDebugTimeStamp)
+		SLATE_ARGUMENT(bool, IsReadOnly)
+
+	SLATE_END_ARGS()
+
+	void Construct(const FArguments& InArgs, UBlackboardData* InBlackboardData);
+
+	/** FGCObject implementation */
+	virtual void AddReferencedObjects( FReferenceCollector& Collector ) override;
+
+	/**
+	 * Retrieves the blackboard item currently selected by the user.
+	 * 
+	 * @return A pointer to the currently selected entry (nullptr if a category is selected, or nothing at all)
+	 */
+	FBlackboardEntry* GetSelectedEntry() const;
+
+	/**
+	 * Retrieves the blackboard item index currently selected by the user.
+	 * @param	bOutIsInherited		Whether the item's index is in the key or parent key array
+	 * @return The selected index (or INDEX_NONE if nothing is selected)
+	 */
+	int32 GetSelectedEntryIndex(bool& bOutIsInherited) const;
+
+	/** Set the object we are looking at */
+	void SetObject(UBlackboardData* InBlackboardData);
+
+protected:
+	/** Delegate handler used to generate a widget for an 'action' (key) in the list */
+	TSharedRef<SWidget> HandleCreateWidgetForAction(FCreateWidgetForActionData* const InCreateData);
+
+	/** Delegate handler used to collect all 'actions' (keys) for display */
+	void HandleCollectAllActions( FGraphActionListBuilderBase& GraphActionListBuilder );
+
+	/** Get the title of the specified section ID */
+	FText HandleGetSectionTitle(int32 SectionID) const;
+
+	/** Delegate handler used when an action is selected */
+	void HandleActionSelected(const TArray< TSharedPtr<FEdGraphSchemaAction> >& SelectedActions) const;
+
+	/** Delegate handler used to generate an action's context menu */
+	TSharedPtr<SWidget> HandleContextMenuOpening(TSharedRef<FUICommandList> ToolkitCommands) const;
+
+	/** Add any context menu items */
+	virtual void FillContextMenu(FMenuBuilder& MenuBuilder) const;
+
+	/** Extend the toolbar */
+	virtual TSharedPtr<FExtender> GetToolbarExtender(TSharedRef<FUICommandList> ToolkitCommands) const;
+
+	/** Delegate handler invoked when we want to use current values */
+	void HandleUseCurrentValues();
+
+	/** Delegate handler invoked when we want to use saved values */
+	void HandleUseSavedValues();
+
+	/** Delegate handler used by toolbar to get display text for timestamp */
+	FText GetDebugTimeStampText() const;
+
+	/** Get the visibility of the toolbar sections used when debugging */
+	EVisibility GetDebuggerToolbarVisibility() const;
+
+	/** Delegate handler used by toolbar to determine whether to display saved values */
+	bool IsUsingCurrentValues() const;
+
+	/** Delegate handler used by toolbar to determine whether to display current values */
+	bool IsUsingSavedValues() const;
+
+	/** Check whether we have any items selected */
+	bool HasSelectedItems() const;
+
+	/** Helper function for GetSelectedEntry() */
+	TSharedPtr<FEdGraphSchemaAction_BlackboardEntry> GetSelectedEntryInternal() const;
+	
+	/** Delegate handler that shows various controls when the debugger is active */
+	bool IsDebuggerActive() const;
+
+	/** Delegate handler that enables various controls when the debugger is active and paused */
+	bool IsDebuggerPaused() const;
+
+	/** Delegate handler used to match an FName to an action in the list, used for renaming keys */
+	bool HandleActionMatchesName(FEdGraphSchemaAction* InAction, const FName& InName) const;
+
+protected:
+	/** The blackboard we are editing/viewing */
+	UBlackboardData* BlackboardData;
+
+	/** The list of blackboard entries */
+	TSharedPtr<SGraphActionMenu> GraphActionMenu;
+
+	/** Delegate executed when an entry is selected */
+	FOnEntrySelected OnEntrySelected;
+
+	/** Delegate used to retrieve debug data to display */
+	FOnGetDebugKeyValue OnGetDebugKeyValue;
+
+	/** Delegate used to determine whether the BT debugger displaying the current state */
+	FOnGetDisplayCurrentState OnGetDisplayCurrentState;
+
+	/** Delegate used to determine whether the BT debugger is active */
+	FOnIsDebuggerReady OnIsDebuggerReady;
+
+	/** Delegate used to determine whether the BT debugger is paused */
+	FOnIsDebuggerPaused OnIsDebuggerPaused;
+
+	/** Delegate used to get the debugger's current timestamp */
+	FOnGetDebugTimeStamp OnGetDebugTimeStamp;
+
+	/** Whether we want to show the current or saved state */
+	bool bShowCurrentState;
+};
\ No newline at end of file
diff --git a/Engine/Source/Editor/BehaviorTreeEditor/Private/SBehaviorTreeDebuggerView.cpp b/Engine/Source/Editor/BehaviorTreeEditor/Private/SBehaviorTreeDebuggerView.cpp
deleted file mode 100644
diff --git a/Engine/Source/Editor/BehaviorTreeEditor/Private/SBehaviorTreeDebuggerView.h b/Engine/Source/Editor/BehaviorTreeEditor/Private/SBehaviorTreeDebuggerView.h
deleted file mode 100644
diff --git a/Engine/Source/Editor/ClassViewer/Private/ClassViewerNode.cpp b/Engine/Source/Editor/ClassViewer/Private/ClassViewerNode.cpp
--- a/Engine/Source/Editor/ClassViewer/Private/ClassViewerNode.cpp
+++ b/Engine/Source/Editor/ClassViewer/Private/ClassViewerNode.cpp
@@ -7,11 +7,12 @@
 #include "PropertyEditorModule.h"
 #include "PropertyHandle.h"
 
-FClassViewerNode::FClassViewerNode( FString _ClassName, bool _bIsPlaceable)
+FClassViewerNode::FClassViewerNode( const FString& InClassName, const FString& InClassDisplayName, bool bInIsPlaceable )
 {
-	ClassName = MakeShareable(new FString(_ClassName));
+	ClassName = MakeShareable(new FString(InClassName));
+	ClassDisplayName = MakeShareable(new FString(InClassDisplayName));
 	bPassesFilter = false;
-	bIsClassPlaceable = _bIsPlaceable;
+	bIsClassPlaceable = bInIsPlaceable;
 	bIsBPNormalType = false;
 
 	Class = NULL;
@@ -21,6 +22,7 @@ FClassViewerNode::FClassViewerNode( FString _ClassName, bool _bIsPlaceable)
 FClassViewerNode::FClassViewerNode( const FClassViewerNode& InCopyObject)
 {
 	ClassName = InCopyObject.ClassName;
+	ClassDisplayName = InCopyObject.ClassDisplayName;
 	bPassesFilter = InCopyObject.bPassesFilter;
 	bIsClassPlaceable = InCopyObject.bIsClassPlaceable;
 
diff --git a/Engine/Source/Editor/ClassViewer/Private/ClassViewerNode.h b/Engine/Source/Editor/ClassViewer/Private/ClassViewerNode.h
--- a/Engine/Source/Editor/ClassViewer/Private/ClassViewerNode.h
+++ b/Engine/Source/Editor/ClassViewer/Private/ClassViewerNode.h
@@ -8,10 +8,11 @@ public:
 	/**
 	 * Creates a node for the widget's tree.
 	 *
-	 * @param	_ClassName						The name of the class this node represents.
-	 * @param	_bIsPlaceable					true if the class is a placeable class.
+	 * @param	InClassName						The name of the class this node represents.
+	 * @param	InClassDisplayName				The display name of the class this node represents
+	 * @param	bInIsPlaceable					true if the class is a placeable class.
 	 */
-	FClassViewerNode( FString Item, bool bIsPlaceable);
+	FClassViewerNode( const FString& InClassName, const FString& InClassDisplayName, bool bInIsPlaceable );
 
 	FClassViewerNode( const FClassViewerNode& InCopyObject);
 
@@ -30,10 +31,13 @@ public:
 	 */
 	void AddUniqueChild(TSharedPtr<FClassViewerNode> NewChild);
 
-	/** Retrieves the class name this node is associated with. */
-	TSharedPtr<FString> GetClassName() const
+	/** 
+	 * Retrieves the class name this node is associated with. 
+	 * @param	bUseDisplayName	Whether to use the display name or class name
+	 */
+	TSharedPtr<FString> GetClassName(bool bUseDisplayName = false) const
 	{
-		return ClassName;	
+		return bUseDisplayName ? ClassDisplayName : ClassName;
 	}
 
 	/** Retrieves the children list. */
@@ -59,6 +63,9 @@ private:
 	/** The class name for this tree node. */
 	TSharedPtr<FString> ClassName;
 
+	/** The class display name for this tree node. */
+	TSharedPtr<FString> ClassDisplayName;
+
 	/** List of children. */
 	TArray<TSharedPtr<FClassViewerNode>> ChildrenList;
 
diff --git a/Engine/Source/Editor/ClassViewer/Private/SClassViewer.cpp b/Engine/Source/Editor/ClassViewer/Private/SClassViewer.cpp
--- a/Engine/Source/Editor/ClassViewer/Private/SClassViewer.cpp
+++ b/Engine/Source/Editor/ClassViewer/Private/SClassViewer.cpp
@@ -1553,14 +1553,12 @@ FClassHierarchy::~FClassHierarchy()
 
 static TSharedPtr< FClassViewerNode > CreateNodeForClass(UClass* Class, const TMultiMap<FName, FAssetData>& BlueprintPackageToAssetDataMap)
 {
-	FString ClassName = Class->GetName();
-
 	const bool bIsPlaceable = ClassViewer::Helpers::IsPlaceable(Class);
 	const bool bIsAbstract = ClassViewer::Helpers::IsAbstract(Class);
 	const bool bIsBrush = ClassViewer::Helpers::IsBrush(Class);
 
 	// Create the new node so it can be passed to AddChildren, fill it in with if it is placeable, abstract, and/or a brush.
-	TSharedPtr< FClassViewerNode > NewNode = MakeShareable(new FClassViewerNode(ClassName, bIsPlaceable && !bIsAbstract && !bIsBrush));
+	TSharedPtr< FClassViewerNode > NewNode = MakeShareable(new FClassViewerNode(Class->GetName(), Class->GetDisplayNameText().ToString(), bIsPlaceable && !bIsAbstract && !bIsBrush));
 	NewNode->Blueprint = ClassViewer::Helpers::GetBlueprint(Class);
 	NewNode->Class = Class;
 
@@ -1587,7 +1585,7 @@ void FClassHierarchy::AddChildren_NoFilter( TSharedPtr< FClassViewerNode >& InOu
 {
 	UClass* RootClass = UObject::StaticClass();
 
-	ObjectClassRoot = MakeShareable(new FClassViewerNode(RootClass->GetName(), false));
+	ObjectClassRoot = MakeShareable(new FClassViewerNode(RootClass->GetName(), RootClass->GetDisplayNameText().ToString(), false));
 	ObjectClassRoot->Class = RootClass;
 
 	TMap< UClass*, TSharedPtr< FClassViewerNode > > Nodes;
@@ -1840,7 +1838,7 @@ void FClassHierarchy::LoadUnloadedTagData(TSharedPtr<FClassViewerNode>& InOutCla
 	const FString AssetName = InAssetData.AssetName.ToString();
 
 	// Create the viewer node.
-	InOutClassViewerNode = MakeShareable(new FClassViewerNode(AssetName, true));
+	InOutClassViewerNode = MakeShareable(new FClassViewerNode(AssetName, AssetName, true));
 			
 	// It is an unloaded blueprint, so we need to create the structure that will hold the data.
 	TSharedPtr<FUnloadedBlueprintData> UnloadedBlueprintData = MakeShareable( new FUnloadedBlueprintData(InOutClassViewerNode) );
@@ -2320,7 +2318,7 @@ TSharedRef< ITableRow > SClassViewer::OnGenerateRowForClassViewer( TSharedPtr<FC
 	// If the item was accepted by the filter, leave it bright, otherwise dim it.
 	float AlphaValue = Item->bPassesFilter? 1.0f : 0.5f;
 	TSharedRef< SClassItem > ReturnRow = SNew(SClassItem, OwnerTable)
-		.ClassName(Item->GetClassName())
+		.ClassName(Item->GetClassName(InitOptions.bShowDisplayNames))
 		.bIsPlaceable(Item->IsClassPlaceable())
 		.HighlightText(&SearchBox->GetText())
 		.TextColor(Item->IsClassPlaceable()? FLinearColor(0.2f, 0.4f, 0.6f, AlphaValue) : FLinearColor(1.0f, 1.0f, 1.0f, AlphaValue))
@@ -2803,7 +2801,7 @@ FReply SClassViewer::OnKeyboardFocusReceived( const FGeometry& MyGeometry, const
 
 TSharedPtr<FClassViewerNode> SClassViewer::CreateNoneOption()
 {
-	TSharedPtr<FClassViewerNode> NoneItem = MakeShareable( new FClassViewerNode("None", false) );
+	TSharedPtr<FClassViewerNode> NoneItem = MakeShareable( new FClassViewerNode("None", "None", false) );
 
 	// The item "passes" the filter so it does not appear grayed out.
 	NoneItem->bPassesFilter = true;
diff --git a/Engine/Source/Editor/ClassViewer/Public/ClassViewerModule.h b/Engine/Source/Editor/ClassViewer/Public/ClassViewerModule.h
--- a/Engine/Source/Editor/ClassViewer/Public/ClassViewerModule.h
+++ b/Engine/Source/Editor/ClassViewer/Public/ClassViewerModule.h
@@ -73,6 +73,9 @@ public:
 	/** true allows class dynamic loading on selection */
 	bool bEnableClassDynamicLoading;
 
+	/** true shows display names of classes rather than full class names */
+	bool bShowDisplayNames;
+
 	/** the title string of the class viewer if required. */
 	FString ViewerTitleString;
 
@@ -92,6 +95,7 @@ public:
 		, bShowObjectRootClass(false)
 		, bExpandRootNodes(true)
 		, bEnableClassDynamicLoading(true)
+		, bShowDisplayNames(false)
 	{
 	}
 };
diff --git a/Engine/Source/Editor/DetailCustomizations/Private/BlackboardEntryDetails.cpp b/Engine/Source/Editor/DetailCustomizations/Private/BlackboardEntryDetails.cpp
--- a/Engine/Source/Editor/DetailCustomizations/Private/BlackboardEntryDetails.cpp
+++ b/Engine/Source/Editor/DetailCustomizations/Private/BlackboardEntryDetails.cpp
@@ -14,21 +14,14 @@ TSharedRef<IPropertyTypeCustomization> FBlackboardEntryDetails::MakeInstance()
 BEGIN_SLATE_FUNCTION_BUILD_OPTIMIZATION
 void FBlackboardEntryDetails::CustomizeHeader( TSharedRef<class IPropertyHandle> StructPropertyHandle, class FDetailWidgetRow& HeaderRow, IPropertyTypeCustomizationUtils& StructCustomizationUtils )
 {
-	MyStructProperty = StructPropertyHandle;
-	MyNameProperty = MyStructProperty->GetChildHandle(TEXT("EntryName"));
-	MyValueProperty = MyStructProperty->GetChildHandle(TEXT("KeyType"));
-
-	HeaderRow.NameContent()
-		[
-			StructPropertyHandle->CreatePropertyNameWidget()
-		]
-		.ValueContent()
-		.MaxDesiredWidth(400.0f)
-		[
-			SNew(STextBlock)
-			.Text(this, &FBlackboardEntryDetails::GetHeaderDesc)
-			.Font(IDetailLayoutBuilder::GetDetailFont())
-		];
+	MyNameProperty = StructPropertyHandle->GetChildHandle(TEXT("EntryName"));
+	MyValueProperty = StructPropertyHandle->GetChildHandle(TEXT("KeyType"));
+
+	// dont show a header row
+	HeaderRow.WholeRowContent()
+	[
+		SNullWidget::NullWidget
+	];
 }
 END_SLATE_FUNCTION_BUILD_OPTIMIZATION
 
@@ -41,46 +34,4 @@ void FBlackboardEntryDetails::CustomizeChildren( TSharedRef<class IPropertyHandl
 	}
 }
 
-FString FBlackboardEntryDetails::GetShortTypeName(const UObject* Ob) const
-{
-	if (Ob->GetClass()->HasAnyClassFlags(CLASS_CompiledFromBlueprint))
-	{
-		return Ob->GetClass()->GetName().LeftChop(2);
-	}
-
-	FString TypeDesc = Ob->GetClass()->GetName();
-	const int32 ShortNameIdx = TypeDesc.Find(TEXT("_"));
-	if (ShortNameIdx != INDEX_NONE)
-	{
-		TypeDesc = TypeDesc.Mid(ShortNameIdx + 1);
-	}
-
-	return TypeDesc;
-}
-
-FString FBlackboardEntryDetails::GetHeaderDesc() const
-{
-	FString StrValue;
-	UObject* ObValue;
-
-	if (MyNameProperty->GetValue(StrValue) == FPropertyAccess::Success &&
-		MyValueProperty->GetValue(ObValue) == FPropertyAccess::Success)
-	{
-		FString ExtraTypeData = ObValue ? ((UBlackboardKeyType*)ObValue)->DescribeSelf() : FString();
-		FString Desc = FString::Printf(TEXT("%s = %s"),
-			StrValue.Len() ? *StrValue : TEXT("??"),
-			ObValue ? *GetShortTypeName(ObValue) : TEXT("??"));
-
-		if (ExtraTypeData.Len())
-		{
-			Desc += TEXT(": ");
-			Desc += ExtraTypeData;
-		}
-
-		return Desc;
-	}
-
-	return FString("??");
-}
-
 #undef LOCTEXT_NAMESPACE
diff --git a/Engine/Source/Editor/DetailCustomizations/Private/BlackboardEntryDetails.h b/Engine/Source/Editor/DetailCustomizations/Private/BlackboardEntryDetails.h
--- a/Engine/Source/Editor/DetailCustomizations/Private/BlackboardEntryDetails.h
+++ b/Engine/Source/Editor/DetailCustomizations/Private/BlackboardEntryDetails.h
@@ -13,11 +13,7 @@ public:
 	virtual void CustomizeChildren( TSharedRef<class IPropertyHandle> StructPropertyHandle, class IDetailChildrenBuilder& StructBuilder, IPropertyTypeCustomizationUtils& StructCustomizationUtils ) override;
 
 private:
-
-	FString GetHeaderDesc() const;
-	FString GetShortTypeName(const UObject* Ob) const;
 	
-	TSharedPtr<IPropertyHandle> MyStructProperty;
 	TSharedPtr<IPropertyHandle> MyNameProperty;
 	TSharedPtr<IPropertyHandle> MyValueProperty;
 };
diff --git a/Engine/Source/Editor/EditorStyle/Private/SlateEditorStyle.cpp b/Engine/Source/Editor/EditorStyle/Private/SlateEditorStyle.cpp
--- a/Engine/Source/Editor/EditorStyle/Private/SlateEditorStyle.cpp
+++ b/Engine/Source/Editor/EditorStyle/Private/SlateEditorStyle.cpp
@@ -2838,6 +2838,11 @@ void FSlateEditorStyle::FStyle::SetupGraphEditorStyles()
 			Set( "BTEditor.Debugger.StopPlaySession", new IMAGE_BRUSH( "Icons/icon_stop_40x", Icon40x40 ) );
 			Set( "BTEditor.Debugger.StopPlaySession.Small", new IMAGE_BRUSH( "Icons/icon_stop_40x", Icon20x20 ) );
 
+			Set( "BTEditor.Debugger.CurrentValues", new IMAGE_BRUSH( "BehaviorTree/Debugger_Current_40x", Icon40x40 ) );
+			Set( "BTEditor.Debugger.CurrentValues.Small", new IMAGE_BRUSH( "BehaviorTree/Debugger_Current_40x", Icon20x20 ) );
+			Set( "BTEditor.Debugger.SavedValues", new IMAGE_BRUSH( "BehaviorTree/Debugger_Saved_40x", Icon40x40 ) );
+			Set( "BTEditor.Debugger.SavedValues.Small", new IMAGE_BRUSH( "BehaviorTree/Debugger_Saved_40x", Icon20x20 ) );
+
 			Set( "BTEditor.DebuggerOverlay.Breakpoint.Disabled", new IMAGE_BRUSH( "Old/Kismet2/Breakpoint_Disabled", Icon32x32 ) );
 			Set( "BTEditor.DebuggerOverlay.Breakpoint.Enabled", new IMAGE_BRUSH( "Old/Kismet2/Breakpoint_Valid", Icon32x32 ) );
 			Set( "BTEditor.DebuggerOverlay.ActiveNodePointer", new IMAGE_BRUSH( "Old/Kismet2/IP_Normal", FVector2D(128,96)) );
@@ -2845,11 +2850,26 @@ void FSlateEditorStyle::FStyle::SetupGraphEditorStyles()
 			Set( "BTEditor.DebuggerOverlay.FailedTriggerPointer", new IMAGE_BRUSH( "/BehaviorTree/FailedTriggerPointer", FVector2D(48,64)) );
 			Set( "BTEditor.DebuggerOverlay.BreakOnBreakpointPointer", new IMAGE_BRUSH( "Old/Kismet2/IP_Breakpoint", FVector2D(128,96)) );
 
+			Set( "BTEditor.Blackboard.NewEntry", new IMAGE_BRUSH( "BehaviorTree/Blackboard_AddKey_40x", Icon40x40 ) );
+			Set( "BTEditor.Blackboard.NewEntry.Small", new IMAGE_BRUSH( "BehaviorTree/Blackboard_AddKey_40x", Icon20x20 ) );
+
 			Set( "BTEditor.SwitchToBehaviorTreeMode", new IMAGE_BRUSH( "BehaviorTree/BehaviorTreeMode_40x", FVector2D(40,40)));
 			Set( "BTEditor.SwitchToBehaviorTreeMode.Small", new IMAGE_BRUSH( "BehaviorTree/BehaviorTreeMode_16x", FVector2D(16,16)));
 			Set( "BTEditor.SwitchToBlackboardMode", new IMAGE_BRUSH( "BehaviorTree/BlackboardMode_40x", FVector2D(40,40)));
 			Set( "BTEditor.SwitchToBlackboardMode.Small", new IMAGE_BRUSH( "BehaviorTree/BlackboardMode_16x", FVector2D(16,16)));
 
+			// Blackboard classes
+			Set( "ClassIcon.BlackboardKeyType_Bool", new IMAGE_BRUSH( "Icons/pill_16x", Icon16x16, FLinearColor(0.300000f, 0.0f, 0.0f, 1.0f) ) );
+			Set( "ClassIcon.BlackboardKeyType_Class", new IMAGE_BRUSH( "Icons/pill_16x", Icon16x16, FLinearColor(0.1f, 0.0f, 0.5f, 1.0f) ) );
+			Set( "ClassIcon.BlackboardKeyType_Enum", new IMAGE_BRUSH( "Icons/pill_16x", Icon16x16, FLinearColor(0.0f, 0.160000f, 0.131270f, 1.0f) ) );
+			Set( "ClassIcon.BlackboardKeyType_Float", new IMAGE_BRUSH( "Icons/pill_16x", Icon16x16, FLinearColor(0.357667f, 1.0f, 0.060000f, 1.0f) ) );
+			Set( "ClassIcon.BlackboardKeyType_Int", new IMAGE_BRUSH( "Icons/pill_16x", Icon16x16, FLinearColor(0.013575f, 0.770000f, 0.429609f, 1.0f) ) );
+			Set( "ClassIcon.BlackboardKeyType_Name", new IMAGE_BRUSH( "Icons/pill_16x", Icon16x16, FLinearColor(0.607717f, 0.224984f, 1.0f, 1.0f) ) );
+			Set( "ClassIcon.BlackboardKeyType_NativeEnum", new IMAGE_BRUSH( "Icons/pill_16x", Icon16x16, FLinearColor(0.0f, 0.160000f, 0.131270f, 1.0f) ) );
+			Set( "ClassIcon.BlackboardKeyType_Object", new IMAGE_BRUSH( "Icons/pill_16x", Icon16x16, FLinearColor(0.0f, 0.4f, 0.910000f, 1.0f) ) );
+			Set( "ClassIcon.BlackboardKeyType_Rotator", new IMAGE_BRUSH( "Icons/pill_16x", Icon16x16, FLinearColor(0.353393f, 0.454175f, 1.0f, 1.0f) ) );
+			Set( "ClassIcon.BlackboardKeyType_String", new IMAGE_BRUSH( "Icons/pill_16x", Icon16x16, FLinearColor(1.0f, 0.0f, 0.660537f, 1.0f) ) );
+			Set( "ClassIcon.BlackboardKeyType_Vector", new IMAGE_BRUSH( "Icons/pill_16x", Icon16x16, FLinearColor(1.0f, 0.591255f, 0.016512f, 1.0f) ) );
 		}
 
 		// Visible on hover button for transition node
diff --git a/Engine/Source/Editor/GraphEditor/Private/SGraphActionMenu.cpp b/Engine/Source/Editor/GraphEditor/Private/SGraphActionMenu.cpp
--- a/Engine/Source/Editor/GraphEditor/Private/SGraphActionMenu.cpp
+++ b/Engine/Source/Editor/GraphEditor/Private/SGraphActionMenu.cpp
@@ -175,6 +175,7 @@ void SGraphActionMenu::Construct( const FArguments& InArgs, bool bIsReadOnly/* =
 	this->OnCanRenameSelectedAction = InArgs._OnCanRenameSelectedAction;
 	this->OnGetSectionTitle = InArgs._OnGetSectionTitle;
 	this->FilteredRootAction = FGraphActionNode::NewCategory(TEXT("FILTEREDROOT"));
+	this->OnActionMatchesName = InArgs._OnActionMatchesName;
 	
 	// If a delegate for filtering text is passed in, assign it so that it will be used instead of the built-in filter box
 	if(InArgs._OnGetFilterText.IsBound())
@@ -372,8 +373,7 @@ bool SGraphActionMenu::SelectItemByName(const FName& ItemName, ESelectInfo::Type
 			FEdGraphSchemaAction* GraphAction = CurrentGraphNode->Actions[0].Get();
 			if (GraphAction)
 			{
-
-				if (GraphActionMenuHelpers::ActionMatchesName(GraphAction, ItemName))
+				if ((OnActionMatchesName.IsBound() && OnActionMatchesName.Execute(GraphAction, ItemName)) || GraphActionMenuHelpers::ActionMatchesName(GraphAction, ItemName))
 				{
 					SelectionNode = GraphNodes[i];
 
@@ -392,7 +392,7 @@ bool SGraphActionMenu::SelectItemByName(const FName& ItemName, ESelectInfo::Type
 
 					if(ChildGraphAction)
 					{
-						if (GraphActionMenuHelpers::ActionMatchesName(ChildGraphAction, ItemName))
+						if ((OnActionMatchesName.IsBound() && OnActionMatchesName.Execute(GraphAction, ItemName)) || GraphActionMenuHelpers::ActionMatchesName(ChildGraphAction, ItemName))
 						{
 							SelectionNode = GraphNodes[i]->Children[ChildIdx];
 
diff --git a/Engine/Source/Editor/GraphEditor/Public/SGraphActionMenu.h b/Engine/Source/Editor/GraphEditor/Public/SGraphActionMenu.h
--- a/Engine/Source/Editor/GraphEditor/Public/SGraphActionMenu.h
+++ b/Engine/Source/Editor/GraphEditor/Public/SGraphActionMenu.h
@@ -89,6 +89,8 @@ public:
 	DECLARE_DELEGATE_RetVal_OneParam( FText, FGetSectionTitle, int32 );
 	/** Delegate to get the filter text */
 	DECLARE_DELEGATE_RetVal( FText, FGetFilterText);
+	/** Delegate to check if an action matches a specified name (used for renaming items etc.) */
+	DECLARE_DELEGATE_RetVal_TwoParams( bool, FOnActionMatchesName, FEdGraphSchemaAction*, const FName& );
 
 	SLATE_BEGIN_ARGS( SGraphActionMenu ) 
 		{
@@ -108,6 +110,7 @@ public:
 		SLATE_EVENT( FCanRenameSelectedAction, OnCanRenameSelectedAction )
 		SLATE_EVENT( FGetSectionTitle, OnGetSectionTitle )		
 		SLATE_EVENT( FGetFilterText, OnGetFilterText )
+		SLATE_EVENT( FOnActionMatchesName, OnActionMatchesName )
 		SLATE_ARGUMENT( bool, AutoExpandActionMenu )
 		SLATE_ARGUMENT( bool, AlphaSortItems )
 		SLATE_ARGUMENT( bool, ShowFilterTextBox )
@@ -173,6 +176,9 @@ protected:
 	FGetSectionTitle OnGetSectionTitle;
 	/** Delegate to get the filter text if supplied from an external source */
 	FGetFilterText OnGetFilterText;
+	/** Delegate to check if an action matches a specified name (used for renaming items etc.) */
+	FOnActionMatchesName OnActionMatchesName;
+
 public:
 	// SWidget interface
 	virtual FReply OnKeyDown(const FGeometry& MyGeometry, const FKeyboardEvent& KeyboardEvent) override;
diff --git a/Engine/Source/Editor/PropertyEditor/Private/UserInterface/PropertyEditor/SPropertyEditorEditInline.cpp b/Engine/Source/Editor/PropertyEditor/Private/UserInterface/PropertyEditor/SPropertyEditorEditInline.cpp
--- a/Engine/Source/Editor/PropertyEditor/Private/UserInterface/PropertyEditor/SPropertyEditorEditInline.cpp
+++ b/Engine/Source/Editor/PropertyEditor/Private/UserInterface/PropertyEditor/SPropertyEditorEditInline.cpp
@@ -10,6 +10,7 @@
 #include "SPropertyComboBox.h"
 #include "Editor/ClassViewer/Public/ClassViewerModule.h"
 #include "Editor/ClassViewer/Public/ClassViewerFilter.h"
+#include "ClassIconFinder.h"
 
 class FPropertyEditorInlineClassFilter : public IClassViewerFilter
 {
@@ -65,6 +66,7 @@ public:
 void SPropertyEditorEditInline::Construct( const FArguments& InArgs, const TSharedRef< class FPropertyEditor >& InPropertyEditor )
 {
 	PropertyEditor = InPropertyEditor;
+	bUseDisplayNames = PropertyEditor->GetProperty()->GetBoolMetaData(TEXT("UseDisplayNames"));
 
 	ChildSlot
 	[
@@ -74,8 +76,23 @@ void SPropertyEditorEditInline::Construct( const FArguments& InArgs, const TShar
 		.ToolTipText(InPropertyEditor, &FPropertyEditor::GetValueAsString )
 		.ButtonContent()
 		[
-			SNew( STextBlock )
-			.Text( this, &SPropertyEditorEditInline::GetDisplayValueAsString )
+			SNew(SHorizontalBox)
+			+SHorizontalBox::Slot()
+			.AutoWidth()
+			.VAlign(VAlign_Center)
+			.Padding(0.0f, 0.0f, 4.0f, 0.0f)
+			[
+				SNew( SImage )
+				.Image( this, &SPropertyEditorEditInline::GetDisplayValueIcon )
+			]
+			+SHorizontalBox::Slot()
+			.AutoWidth()
+			.VAlign(VAlign_Center)
+			[
+				SNew( STextBlock )
+				.Text( this, &SPropertyEditorEditInline::GetDisplayValueAsString )
+				.Font( InArgs._Font )
+			]
 		]
 	];
 }
@@ -86,7 +103,14 @@ FString SPropertyEditorEditInline::GetDisplayValueAsString() const
 	FPropertyAccess::Result Result = PropertyEditor->GetPropertyHandle()->GetValue( CurrentValue );
 	if( Result == FPropertyAccess::Success && CurrentValue != NULL )
 	{
-		return CurrentValue->GetName();
+		if(bUseDisplayNames)
+		{
+			return CurrentValue->GetClass()->GetDisplayNameText().ToString();
+		}
+		else
+		{
+			return CurrentValue->GetName();
+		}
 	}
 	else
 	{
@@ -94,6 +118,18 @@ FString SPropertyEditorEditInline::GetDisplayValueAsString() const
 	}
 }
 
+const FSlateBrush* SPropertyEditorEditInline::GetDisplayValueIcon() const
+{
+	UObject* CurrentValue = nullptr;
+	FPropertyAccess::Result Result = PropertyEditor->GetPropertyHandle()->GetValue( CurrentValue );
+	if( Result == FPropertyAccess::Success && CurrentValue != nullptr )
+	{
+		return FClassIconFinder::FindIconForClass(CurrentValue->GetClass());
+	}
+
+	return nullptr;
+}
+
 void SPropertyEditorEditInline::GetDesiredWidth( float& OutMinDesiredWidth, float& OutMaxDesiredWidth )
 {
 	OutMinDesiredWidth = 250.0f;
@@ -146,6 +182,7 @@ TSharedRef<SWidget> SPropertyEditorEditInline::GenerateClassPicker()
 {
 	FClassViewerInitializationOptions Options;
 	Options.bShowUnloadedBlueprints = true;
+	Options.bShowDisplayNames = bUseDisplayNames;
 
 	TSharedPtr<FPropertyEditorInlineClassFilter> ClassFilter = MakeShareable( new FPropertyEditorInlineClassFilter );
 	Options.ClassFilter = ClassFilter;
diff --git a/Engine/Source/Editor/PropertyEditor/Private/UserInterface/PropertyEditor/SPropertyEditorEditInline.h b/Engine/Source/Editor/PropertyEditor/Private/UserInterface/PropertyEditor/SPropertyEditorEditInline.h
--- a/Engine/Source/Editor/PropertyEditor/Private/UserInterface/PropertyEditor/SPropertyEditorEditInline.h
+++ b/Engine/Source/Editor/PropertyEditor/Private/UserInterface/PropertyEditor/SPropertyEditorEditInline.h
@@ -27,6 +27,11 @@ private:
 	 */
 	FString GetDisplayValueAsString() const;
 
+	/**
+	 * @return The current display value's icon, if any. Returns nullptr if we have no valid value.
+	 */
+	const FSlateBrush* GetDisplayValueIcon() const;
+
 	/**
 	 * Wrapper method for determining whether a class is valid for use by this property item input proxy.
 	 *
@@ -58,4 +63,7 @@ private:
 	TSharedPtr<class FPropertyEditor > PropertyEditor;
 
 	TSharedPtr<class SComboButton> ComboButton;
+
+	/** Whether to use display names or literal class names */
+	bool bUseDisplayNames;
 };
diff --git a/Engine/Source/Editor/PropertyEditor/Public/PropertyEditorModule.h b/Engine/Source/Editor/PropertyEditor/Public/PropertyEditorModule.h
--- a/Engine/Source/Editor/PropertyEditor/Public/PropertyEditorModule.h
+++ b/Engine/Source/Editor/PropertyEditor/Public/PropertyEditorModule.h
@@ -203,7 +203,7 @@ public:
 	 * @param Class	The class the custom detail layout is for
 	 * @param StructLayoutDelegate	The delegate to call when querying for custom detail layouts for the struct properties
 	 */
-	DEPRECATED(4.3, "RegisterCustomPropertyTypeLayout is deprecated.  Use RegisterCustomPropertyTypeLayout instead.")
+	DEPRECATED(4.3, "RegisterStructPropertyLayout is deprecated.  Use RegisterCustomPropertyTypeLayout instead.")
 	virtual void RegisterStructPropertyLayout( FName StructTypeName, FOnGetStructCustomizationInstance StructLayoutDelegate );
 
 	/**
diff --git a/Engine/Source/Runtime/AIModule/Classes/BehaviorTree/BehaviorTreeComponent.h b/Engine/Source/Runtime/AIModule/Classes/BehaviorTree/BehaviorTreeComponent.h
--- a/Engine/Source/Runtime/AIModule/Classes/BehaviorTree/BehaviorTreeComponent.h
+++ b/Engine/Source/Runtime/AIModule/Classes/BehaviorTree/BehaviorTreeComponent.h
@@ -162,7 +162,7 @@ protected:
 	/** search flow for debugger */
 	mutable TArray<TArray<FBehaviorTreeDebuggerInstance::FNodeFlowData> > CurrentSearchFlow;
 	mutable TArray<TArray<FBehaviorTreeDebuggerInstance::FNodeFlowData> > CurrentRestarts;
-	mutable TArray<FString> SearchStartBlackboard;
+	mutable TMap<FName, FString> SearchStartBlackboard;
 	mutable TArray<FBehaviorTreeDebuggerInstance> RemovedInstances;
 
 	/** debugger's recorded data */
@@ -233,7 +233,7 @@ protected:
 	void StoreDebuggerRestart(const class UBTNode* Node, uint16 InstanceIdx, bool bAllowed);
 
 	/** describe blackboard's key values */
-	void StoreDebuggerBlackboard(TArray<FString>& BlackboardValueDesc) const;
+	void StoreDebuggerBlackboard(TMap<FName, FString>& BlackboardValueDesc) const;
 
 	/** gather nodes runtime descriptions */
 	void StoreDebuggerRuntimeValues(TArray<FString>& RuntimeDescriptions, class UBTNode* RootNode, uint16 InstanceIdx) const;
diff --git a/Engine/Source/Runtime/AIModule/Classes/BehaviorTree/BehaviorTreeTypes.h b/Engine/Source/Runtime/AIModule/Classes/BehaviorTree/BehaviorTreeTypes.h
--- a/Engine/Source/Runtime/AIModule/Classes/BehaviorTree/BehaviorTreeTypes.h
+++ b/Engine/Source/Runtime/AIModule/Classes/BehaviorTree/BehaviorTreeTypes.h
@@ -3,6 +3,8 @@
 #pragma once
 #include "BehaviorTreeTypes.generated.h"
 
+class UBlackboardData;
+
 // Visual logging helper
 #define BT_VLOG(Context, Verbosity, Format, ...) UE_VLOG(Context->OwnerComp.IsValid() ? Context->OwnerComp->GetOwner() : NULL, LogBehaviorTree, Verbosity, Format, ##__VA_ARGS__)
 #define BT_SEARCHLOG(SearchData, Verbosity, Format, ...) UE_VLOG(SearchData.OwnerComp->GetOwner(), LogBehaviorTree, Verbosity, Format, ##__VA_ARGS__)
@@ -189,7 +191,7 @@ struct FBehaviorTreeExecutionStep
 	TArray<FBehaviorTreeDebuggerInstance> InstanceStack;
 
 	/** blackboard snapshot: value descriptions */
-	TArray<FString> BlackboardValues;
+	TMap<FName, FString> BlackboardValues;
 
 	/** Game world's time stamp of this step */
 	float TimeStamp;
@@ -376,10 +378,10 @@ protected:
 
 public:
 	/** cache ID and class of selected key */
-	void CacheSelectedKey(class UBlackboardData* BlackboardAsset);
+	void CacheSelectedKey(UBlackboardData* BlackboardAsset);
 
 	/** find initial selection */
-	void InitSelectedKey(class UBlackboardData* BlackboardAsset);
+	void InitSelectedKey(UBlackboardData* BlackboardAsset);
 
 	void AllowNoneAsValue(bool bNewVal ) { bNoneIsAllowedValue = bNewVal; }
 
diff --git a/Engine/Source/Runtime/AIModule/Classes/BehaviorTree/Blackboard/BlackboardKeyType_Bool.h b/Engine/Source/Runtime/AIModule/Classes/BehaviorTree/Blackboard/BlackboardKeyType_Bool.h
--- a/Engine/Source/Runtime/AIModule/Classes/BehaviorTree/Blackboard/BlackboardKeyType_Bool.h
+++ b/Engine/Source/Runtime/AIModule/Classes/BehaviorTree/Blackboard/BlackboardKeyType_Bool.h
@@ -3,7 +3,7 @@
 #include "BehaviorTree/Blackboard/BlackboardKeyType.h"
 #include "BlackboardKeyType_Bool.generated.h"
 
-UCLASS(EditInlineNew, MinimalAPI)
+UCLASS(EditInlineNew, MinimalAPI, meta=(DisplayName="Bool"))
 class UBlackboardKeyType_Bool : public UBlackboardKeyType
 {
 	GENERATED_UCLASS_BODY()
diff --git a/Engine/Source/Runtime/AIModule/Classes/BehaviorTree/Blackboard/BlackboardKeyType_Class.h b/Engine/Source/Runtime/AIModule/Classes/BehaviorTree/Blackboard/BlackboardKeyType_Class.h
--- a/Engine/Source/Runtime/AIModule/Classes/BehaviorTree/Blackboard/BlackboardKeyType_Class.h
+++ b/Engine/Source/Runtime/AIModule/Classes/BehaviorTree/Blackboard/BlackboardKeyType_Class.h
@@ -3,12 +3,12 @@
 #include "BehaviorTree/Blackboard/BlackboardKeyType.h"
 #include "BlackboardKeyType_Class.generated.h"
 
-UCLASS(EditInlineNew, MinimalAPI)
+UCLASS(EditInlineNew, MinimalAPI, meta=(DisplayName="Class"))
 class UBlackboardKeyType_Class : public UBlackboardKeyType
 {
 	GENERATED_UCLASS_BODY()
 
-	UPROPERTY(Category=Blackboard, EditDefaultsOnly, meta=(AllowAbstract="1"))
+	UPROPERTY(Category=Blackboard, EditInline, EditDefaultsOnly, meta=(AllowAbstract="1"))
 	UClass* BaseClass;
 	
 	static UClass* GetValue(const uint8* RawData);
diff --git a/Engine/Source/Runtime/AIModule/Classes/BehaviorTree/Blackboard/BlackboardKeyType_Enum.h b/Engine/Source/Runtime/AIModule/Classes/BehaviorTree/Blackboard/BlackboardKeyType_Enum.h
--- a/Engine/Source/Runtime/AIModule/Classes/BehaviorTree/Blackboard/BlackboardKeyType_Enum.h
+++ b/Engine/Source/Runtime/AIModule/Classes/BehaviorTree/Blackboard/BlackboardKeyType_Enum.h
@@ -3,12 +3,12 @@
 #include "BehaviorTree/Blackboard/BlackboardKeyType.h"
 #include "BlackboardKeyType_Enum.generated.h"
 
-UCLASS(EditInlineNew, MinimalAPI)
+UCLASS(EditInlineNew, MinimalAPI, meta=(DisplayName="Enum"))
 class UBlackboardKeyType_Enum : public UBlackboardKeyType
 {
 	GENERATED_UCLASS_BODY()
 
-	UPROPERTY(Category=Blackboard, EditDefaultsOnly)
+	UPROPERTY(Category=Blackboard, EditInline, EditDefaultsOnly)
 	UEnum* EnumType;
 
 	static uint8 GetValue(const uint8* RawData);
diff --git a/Engine/Source/Runtime/AIModule/Classes/BehaviorTree/Blackboard/BlackboardKeyType_Float.h b/Engine/Source/Runtime/AIModule/Classes/BehaviorTree/Blackboard/BlackboardKeyType_Float.h
--- a/Engine/Source/Runtime/AIModule/Classes/BehaviorTree/Blackboard/BlackboardKeyType_Float.h
+++ b/Engine/Source/Runtime/AIModule/Classes/BehaviorTree/Blackboard/BlackboardKeyType_Float.h
@@ -3,7 +3,7 @@
 #include "BehaviorTree/Blackboard/BlackboardKeyType.h"
 #include "BlackboardKeyType_Float.generated.h"
 
-UCLASS(EditInlineNew, MinimalAPI)
+UCLASS(EditInlineNew, MinimalAPI, meta=(DisplayName="Float"))
 class UBlackboardKeyType_Float : public UBlackboardKeyType
 {
 	GENERATED_UCLASS_BODY()
diff --git a/Engine/Source/Runtime/AIModule/Classes/BehaviorTree/Blackboard/BlackboardKeyType_Int.h b/Engine/Source/Runtime/AIModule/Classes/BehaviorTree/Blackboard/BlackboardKeyType_Int.h
--- a/Engine/Source/Runtime/AIModule/Classes/BehaviorTree/Blackboard/BlackboardKeyType_Int.h
+++ b/Engine/Source/Runtime/AIModule/Classes/BehaviorTree/Blackboard/BlackboardKeyType_Int.h
@@ -3,7 +3,7 @@
 #include "BehaviorTree/Blackboard/BlackboardKeyType.h"
 #include "BlackboardKeyType_Int.generated.h"
 
-UCLASS(EditInlineNew, MinimalAPI)
+UCLASS(EditInlineNew, MinimalAPI, meta=(DisplayName="Int"))
 class UBlackboardKeyType_Int : public UBlackboardKeyType
 {
 	GENERATED_UCLASS_BODY()
diff --git a/Engine/Source/Runtime/AIModule/Classes/BehaviorTree/Blackboard/BlackboardKeyType_Name.h b/Engine/Source/Runtime/AIModule/Classes/BehaviorTree/Blackboard/BlackboardKeyType_Name.h
--- a/Engine/Source/Runtime/AIModule/Classes/BehaviorTree/Blackboard/BlackboardKeyType_Name.h
+++ b/Engine/Source/Runtime/AIModule/Classes/BehaviorTree/Blackboard/BlackboardKeyType_Name.h
@@ -3,7 +3,7 @@
 #include "BehaviorTree/Blackboard/BlackboardKeyType.h"
 #include "BlackboardKeyType_Name.generated.h"
 
-UCLASS(EditInlineNew, MinimalAPI)
+UCLASS(EditInlineNew, MinimalAPI, meta=(DisplayName="Name"))
 class UBlackboardKeyType_Name : public UBlackboardKeyType
 {
 	GENERATED_UCLASS_BODY()
diff --git a/Engine/Source/Runtime/AIModule/Classes/BehaviorTree/Blackboard/BlackboardKeyType_NativeEnum.h b/Engine/Source/Runtime/AIModule/Classes/BehaviorTree/Blackboard/BlackboardKeyType_NativeEnum.h
--- a/Engine/Source/Runtime/AIModule/Classes/BehaviorTree/Blackboard/BlackboardKeyType_NativeEnum.h
+++ b/Engine/Source/Runtime/AIModule/Classes/BehaviorTree/Blackboard/BlackboardKeyType_NativeEnum.h
@@ -4,15 +4,15 @@
 #include "BehaviorTree/Blackboard/BlackboardKeyType.h"
 #include "BlackboardKeyType_NativeEnum.generated.h"
 
-UCLASS(EditInlineNew, MinimalAPI)
+UCLASS(EditInlineNew, MinimalAPI, meta=(DisplayName="Native Enum"))
 class UBlackboardKeyType_NativeEnum : public UBlackboardKeyType
 {
 	GENERATED_UCLASS_BODY()
 
-	UPROPERTY(Category=Blackboard, EditDefaultsOnly)
+	UPROPERTY(Category=Blackboard, EditInline, EditDefaultsOnly)
 	FString EnumName;
 
-	UPROPERTY(Category=Blackboard, VisibleDefaultsOnly)
+	UPROPERTY(Category=Blackboard, EditInline, VisibleDefaultsOnly)
 	bool bIsEnumNameValid;
 
 	UPROPERTY()
diff --git a/Engine/Source/Runtime/AIModule/Classes/BehaviorTree/Blackboard/BlackboardKeyType_Object.h b/Engine/Source/Runtime/AIModule/Classes/BehaviorTree/Blackboard/BlackboardKeyType_Object.h
--- a/Engine/Source/Runtime/AIModule/Classes/BehaviorTree/Blackboard/BlackboardKeyType_Object.h
+++ b/Engine/Source/Runtime/AIModule/Classes/BehaviorTree/Blackboard/BlackboardKeyType_Object.h
@@ -3,12 +3,12 @@
 #include "BehaviorTree/Blackboard/BlackboardKeyType.h"
 #include "BlackboardKeyType_Object.generated.h"
 
-UCLASS(EditInlineNew, MinimalAPI)
+UCLASS(EditInlineNew, MinimalAPI, meta=(DisplayName="Object"))
 class UBlackboardKeyType_Object : public UBlackboardKeyType
 {
 	GENERATED_UCLASS_BODY()
 
-	UPROPERTY(Category=Blackboard, EditDefaultsOnly, meta=(AllowAbstract="1"))
+	UPROPERTY(Category=Blackboard, EditInline, EditDefaultsOnly, meta=(AllowAbstract="1"))
 	UClass* BaseClass;
 
 	static UObject* GetValue(const uint8* RawData);
diff --git a/Engine/Source/Runtime/AIModule/Classes/BehaviorTree/Blackboard/BlackboardKeyType_Rotator.h b/Engine/Source/Runtime/AIModule/Classes/BehaviorTree/Blackboard/BlackboardKeyType_Rotator.h
--- a/Engine/Source/Runtime/AIModule/Classes/BehaviorTree/Blackboard/BlackboardKeyType_Rotator.h
+++ b/Engine/Source/Runtime/AIModule/Classes/BehaviorTree/Blackboard/BlackboardKeyType_Rotator.h
@@ -2,7 +2,7 @@
 #pragma once
 #include "BlackboardKeyType_Rotator.generated.h"
 
-UCLASS(EditInlineNew, MinimalAPI)
+UCLASS(EditInlineNew, MinimalAPI, meta=(DisplayName="Rotator"))
 class UBlackboardKeyType_Rotator : public UBlackboardKeyType
 {
 	GENERATED_UCLASS_BODY()
diff --git a/Engine/Source/Runtime/AIModule/Classes/BehaviorTree/Blackboard/BlackboardKeyType_String.h b/Engine/Source/Runtime/AIModule/Classes/BehaviorTree/Blackboard/BlackboardKeyType_String.h
--- a/Engine/Source/Runtime/AIModule/Classes/BehaviorTree/Blackboard/BlackboardKeyType_String.h
+++ b/Engine/Source/Runtime/AIModule/Classes/BehaviorTree/Blackboard/BlackboardKeyType_String.h
@@ -3,7 +3,7 @@
 #include "BehaviorTree/Blackboard/BlackboardKeyType.h"
 #include "BlackboardKeyType_String.generated.h"
 
-UCLASS(EditInlineNew, MinimalAPI)
+UCLASS(EditInlineNew, MinimalAPI, meta=(DisplayName="String"))
 class UBlackboardKeyType_String : public UBlackboardKeyType
 {
 	GENERATED_UCLASS_BODY()
diff --git a/Engine/Source/Runtime/AIModule/Classes/BehaviorTree/Blackboard/BlackboardKeyType_Vector.h b/Engine/Source/Runtime/AIModule/Classes/BehaviorTree/Blackboard/BlackboardKeyType_Vector.h
--- a/Engine/Source/Runtime/AIModule/Classes/BehaviorTree/Blackboard/BlackboardKeyType_Vector.h
+++ b/Engine/Source/Runtime/AIModule/Classes/BehaviorTree/Blackboard/BlackboardKeyType_Vector.h
@@ -3,7 +3,7 @@
 #include "BehaviorTree/Blackboard/BlackboardKeyType.h"
 #include "BlackboardKeyType_Vector.generated.h"
 
-UCLASS(EditInlineNew, MinimalAPI)
+UCLASS(EditInlineNew, MinimalAPI, meta=(DisplayName="Vector"))
 class UBlackboardKeyType_Vector : public UBlackboardKeyType
 {
 	GENERATED_UCLASS_BODY()
diff --git a/Engine/Source/Runtime/AIModule/Classes/BehaviorTree/BlackboardComponent.h b/Engine/Source/Runtime/AIModule/Classes/BehaviorTree/BlackboardComponent.h
--- a/Engine/Source/Runtime/AIModule/Classes/BehaviorTree/BlackboardComponent.h
+++ b/Engine/Source/Runtime/AIModule/Classes/BehaviorTree/BlackboardComponent.h
@@ -17,6 +17,9 @@
 #include "BlackboardData.h"
 #include "BlackboardComponent.generated.h"
 
+class UBlackboardData;
+class UBrainComponent;
+
 namespace EBlackboardDescription
 {
 	enum Type
@@ -68,16 +71,16 @@ class AIMODULE_API UBlackboardComponent : public UActorComponent
 	class UBrainComponent* GetBrainComponent() const;
 
 	/** @return blackboard data asset */
-	class UBlackboardData* GetBlackboardAsset() const;
+	UBlackboardData* GetBlackboardAsset() const;
 
 	/** caches UBrainComponent pointer to be used in communication */
-	void CacheBrainComponent(class UBrainComponent* BrainComponent);
+	void CacheBrainComponent(UBrainComponent* BrainComponent);
 
 	/** setup component for using given blackboard asset */
-	void InitializeBlackboard(class UBlackboardData* NewAsset);
+	void InitializeBlackboard(UBlackboardData* NewAsset);
 	
 	/** @return true if component can be used with specified blackboard asset */
-	bool IsCompatibleWith(class UBlackboardData* TestAsset) const;
+	bool IsCompatibleWith(UBlackboardData* TestAsset) const;
 
 	UFUNCTION(BlueprintCallable, Category="AI|Components|Blackboard")
 	UObject* GetValueAsObject(const FName& KeyName) const;
@@ -192,6 +195,7 @@ class AIMODULE_API UBlackboardComponent : public UActorComponent
 	FString GetDebugInfoString(EBlackboardDescription::Type Mode) const;
 
 	/** get description of value under given key */
+	FString DescribeKeyValue(const FName& KeyName, EBlackboardDescription::Type Mode) const;
 	FString DescribeKeyValue(FBlackboard::FKey KeyID, EBlackboardDescription::Type Mode) const;
 
 #if ENABLE_VISUAL_LOG
@@ -203,11 +207,11 @@ protected:
 
 	/** cached behavior tree component */
 	UPROPERTY(transient)
-	class UBrainComponent* BrainComp;
+	UBrainComponent* BrainComp;
 
 	/** data asset defining entries */
 	UPROPERTY(transient)
-	class UBlackboardData* BlackboardAsset;
+	UBlackboardData* BlackboardAsset;
 
 	/** memory block holding all values */
 	TArray<uint8> ValueMemory;
diff --git a/Engine/Source/Runtime/AIModule/Classes/BehaviorTree/BlackboardData.h b/Engine/Source/Runtime/AIModule/Classes/BehaviorTree/BlackboardData.h
--- a/Engine/Source/Runtime/AIModule/Classes/BehaviorTree/BlackboardData.h
+++ b/Engine/Source/Runtime/AIModule/Classes/BehaviorTree/BlackboardData.h
@@ -4,6 +4,8 @@
 #include "Engine/DataAsset.h"
 #include "BlackboardData.generated.h"
 
+class UBlackboardData;
+
 /** blackboard entry definition */
 USTRUCT()
 struct FBlackboardEntry
@@ -14,11 +16,11 @@ struct FBlackboardEntry
 	FName EntryName;
 
 	/** key type and additional properties */
-	UPROPERTY(EditAnywhere, EditInline, Category=Blackboard)
+	UPROPERTY(EditAnywhere, EditInline, Category=Blackboard, meta=(UseDisplayNames="true"))
 	UBlackboardKeyType* KeyType;
 };
 
-UCLASS()
+UCLASS(AutoExpandCategories=(Blackboard))
 class AIMODULE_API UBlackboardData : public UDataAsset
 {
 	GENERATED_UCLASS_BODY()
@@ -26,7 +28,7 @@ class AIMODULE_API UBlackboardData : public UDataAsset
 
 	/** parent blackboard (keys can be overridden) */
 	UPROPERTY(EditAnywhere, Category=Parent)
-	class UBlackboardData* Parent;
+	UBlackboardData* Parent;
 
 #if WITH_EDITORONLY_DATA
 	/** all keys inherited from parent chain */
diff --git a/Engine/Source/Runtime/AIModule/Private/BehaviorTree/BTFunctionLibrary.cpp b/Engine/Source/Runtime/AIModule/Private/BehaviorTree/BTFunctionLibrary.cpp
--- a/Engine/Source/Runtime/AIModule/Private/BehaviorTree/BTFunctionLibrary.cpp
+++ b/Engine/Source/Runtime/AIModule/Private/BehaviorTree/BTFunctionLibrary.cpp
@@ -18,151 +18,151 @@ UBlackboardComponent* UBTFunctionLibrary::GetBlackboard(UBTNode* NodeOwner)
 
 UObject* UBTFunctionLibrary::GetBlackboardValueAsObject(UBTNode* NodeOwner, const struct FBlackboardKeySelector& Key)
 {
-	const UBehaviorTreeComponent* BTComp = NodeOwner ? Cast<UBehaviorTreeComponent>(NodeOwner->GetOuter()) : NULL;
-	return BTComp ? BTComp->GetBlackboardComponent()->GetValueAsObject(Key.SelectedKeyName) : NULL;
+	const UBlackboardComponent* BlackboardComp = GetBlackboard(NodeOwner);
+	return BlackboardComp ? BlackboardComp->GetValueAsObject(Key.SelectedKeyName) : NULL;
 }
 
 AActor* UBTFunctionLibrary::GetBlackboardValueAsActor(UBTNode* NodeOwner, const struct FBlackboardKeySelector& Key)
 {
-	const UBehaviorTreeComponent* BTComp = NodeOwner ? Cast<UBehaviorTreeComponent>(NodeOwner->GetOuter()) : NULL;
-	return BTComp ? Cast<AActor>(BTComp->GetBlackboardComponent()->GetValueAsObject(Key.SelectedKeyName)) : NULL;
+	const UBlackboardComponent* BlackboardComp = GetBlackboard(NodeOwner);
+	return BlackboardComp ? Cast<AActor>(BlackboardComp->GetValueAsObject(Key.SelectedKeyName)) : NULL;
 }
 
 UClass* UBTFunctionLibrary::GetBlackboardValueAsClass(UBTNode* NodeOwner, const struct FBlackboardKeySelector& Key)
 {
-	const UBehaviorTreeComponent* BTComp = NodeOwner ? Cast<UBehaviorTreeComponent>(NodeOwner->GetOuter()) : NULL;
-	return BTComp ? BTComp->GetBlackboardComponent()->GetValueAsClass(Key.SelectedKeyName) : NULL;
+	const UBlackboardComponent* BlackboardComp = GetBlackboard(NodeOwner);
+	return BlackboardComp ? BlackboardComp->GetValueAsClass(Key.SelectedKeyName) : NULL;
 }
 
 uint8 UBTFunctionLibrary::GetBlackboardValueAsEnum(UBTNode* NodeOwner, const struct FBlackboardKeySelector& Key)
 {
-	const UBehaviorTreeComponent* BTComp = NodeOwner ? Cast<UBehaviorTreeComponent>(NodeOwner->GetOuter()) : NULL;
-	return BTComp ? BTComp->GetBlackboardComponent()->GetValueAsEnum(Key.SelectedKeyName) : 0;
+	const UBlackboardComponent* BlackboardComp = GetBlackboard(NodeOwner);
+	return BlackboardComp ? BlackboardComp->GetValueAsEnum(Key.SelectedKeyName) : 0;
 }
 
 int32 UBTFunctionLibrary::GetBlackboardValueAsInt(UBTNode* NodeOwner, const struct FBlackboardKeySelector& Key)
 {
-	const UBehaviorTreeComponent* BTComp = NodeOwner ? Cast<UBehaviorTreeComponent>(NodeOwner->GetOuter()) : NULL;
-	return BTComp ? BTComp->GetBlackboardComponent()->GetValueAsInt(Key.SelectedKeyName) : 0;
+	const UBlackboardComponent* BlackboardComp = GetBlackboard(NodeOwner);
+	return BlackboardComp ? BlackboardComp->GetValueAsInt(Key.SelectedKeyName) : 0;
 }
 
 float UBTFunctionLibrary::GetBlackboardValueAsFloat(UBTNode* NodeOwner, const struct FBlackboardKeySelector& Key)
 {
-	const UBehaviorTreeComponent* BTComp = NodeOwner ? Cast<UBehaviorTreeComponent>(NodeOwner->GetOuter()) : NULL;
-	return BTComp ? BTComp->GetBlackboardComponent()->GetValueAsFloat(Key.SelectedKeyName) : 0.0f;
+	const UBlackboardComponent* BlackboardComp = GetBlackboard(NodeOwner);
+	return BlackboardComp ? BlackboardComp->GetValueAsFloat(Key.SelectedKeyName) : 0.0f;
 }
 
 bool UBTFunctionLibrary::GetBlackboardValueAsBool(UBTNode* NodeOwner, const struct FBlackboardKeySelector& Key)
 {
-	const UBehaviorTreeComponent* BTComp = NodeOwner ? Cast<UBehaviorTreeComponent>(NodeOwner->GetOuter()) : NULL;
-	return BTComp ? BTComp->GetBlackboardComponent()->GetValueAsBool(Key.SelectedKeyName) : false;
+	const UBlackboardComponent* BlackboardComp = GetBlackboard(NodeOwner);
+	return BlackboardComp ? BlackboardComp->GetValueAsBool(Key.SelectedKeyName) : false;
 }
 
 FString UBTFunctionLibrary::GetBlackboardValueAsString(UBTNode* NodeOwner, const struct FBlackboardKeySelector& Key)
 {
-	const UBehaviorTreeComponent* BTComp = NodeOwner ? Cast<UBehaviorTreeComponent>(NodeOwner->GetOuter()) : NULL;
-	return BTComp ? BTComp->GetBlackboardComponent()->GetValueAsString(Key.SelectedKeyName) : FString();
+	const UBlackboardComponent* BlackboardComp = GetBlackboard(NodeOwner);
+	return BlackboardComp ? BlackboardComp->GetValueAsString(Key.SelectedKeyName) : FString();
 }
 
 FName UBTFunctionLibrary::GetBlackboardValueAsName(UBTNode* NodeOwner, const struct FBlackboardKeySelector& Key)
 {
-	const UBehaviorTreeComponent* BTComp = NodeOwner ? Cast<UBehaviorTreeComponent>(NodeOwner->GetOuter()) : NULL;
-	return BTComp ? BTComp->GetBlackboardComponent()->GetValueAsName(Key.SelectedKeyName) : NAME_None;
+	const UBlackboardComponent* BlackboardComp = GetBlackboard(NodeOwner);
+	return BlackboardComp ? BlackboardComp->GetValueAsName(Key.SelectedKeyName) : NAME_None;
 }
 
 FVector UBTFunctionLibrary::GetBlackboardValueAsVector(UBTNode* NodeOwner, const struct FBlackboardKeySelector& Key)
 {
-	const UBehaviorTreeComponent* BTComp = NodeOwner ? Cast<UBehaviorTreeComponent>(NodeOwner->GetOuter()) : NULL;
-	return BTComp ? BTComp->GetBlackboardComponent()->GetValueAsVector(Key.SelectedKeyName) : FVector::ZeroVector;
+	const UBlackboardComponent* BlackboardComp = GetBlackboard(NodeOwner);
+	return BlackboardComp ? BlackboardComp->GetValueAsVector(Key.SelectedKeyName) : FVector::ZeroVector;
 }
 
 void UBTFunctionLibrary::SetBlackboardValueAsObject(UBTNode* NodeOwner, const struct FBlackboardKeySelector& Key, UObject* Value)
 {
-	UBehaviorTreeComponent* BTComp = NodeOwner ? Cast<UBehaviorTreeComponent>(NodeOwner->GetOuter()) : NULL;
-	if (BTComp != NULL)
+	UBlackboardComponent* BlackboardComp = GetBlackboard(NodeOwner);
+	if (BlackboardComp != NULL)
 	{
-		BTComp->GetBlackboardComponent()->SetValueAsObject(Key.SelectedKeyName, Value);
+		BlackboardComp->SetValueAsObject(Key.SelectedKeyName, Value);
 	}
 }
 
 void UBTFunctionLibrary::SetBlackboardValueAsClass(UBTNode* NodeOwner, const struct FBlackboardKeySelector& Key, UClass* Value)
 {
-	UBehaviorTreeComponent* BTComp = NodeOwner ? Cast<UBehaviorTreeComponent>(NodeOwner->GetOuter()) : NULL;
-	if (BTComp != NULL)
+	UBlackboardComponent* BlackboardComp = GetBlackboard(NodeOwner);
+	if (BlackboardComp != NULL)
 	{
-		BTComp->GetBlackboardComponent()->SetValueAsClass(Key.SelectedKeyName, Value);
+		BlackboardComp->SetValueAsClass(Key.SelectedKeyName, Value);
 	}
 }
 
 void UBTFunctionLibrary::SetBlackboardValueAsEnum(UBTNode* NodeOwner, const struct FBlackboardKeySelector& Key, uint8 Value)
 {
-	UBehaviorTreeComponent* BTComp = NodeOwner ? Cast<UBehaviorTreeComponent>(NodeOwner->GetOuter()) : NULL;
-	if (BTComp != NULL)
+	UBlackboardComponent* BlackboardComp = GetBlackboard(NodeOwner);
+	if (BlackboardComp != NULL)
 	{
-		BTComp->GetBlackboardComponent()->SetValueAsEnum(Key.SelectedKeyName, Value);
+		BlackboardComp->SetValueAsEnum(Key.SelectedKeyName, Value);
 	}
 }
 
 void UBTFunctionLibrary::SetBlackboardValueAsInt(UBTNode* NodeOwner, const struct FBlackboardKeySelector& Key, int32 Value)
 {
-	UBehaviorTreeComponent* BTComp = NodeOwner ? Cast<UBehaviorTreeComponent>(NodeOwner->GetOuter()) : NULL;
-	if (BTComp != NULL)
+	UBlackboardComponent* BlackboardComp = GetBlackboard(NodeOwner);
+	if (BlackboardComp != NULL)
 	{
-		BTComp->GetBlackboardComponent()->SetValueAsInt(Key.SelectedKeyName, Value);
+		BlackboardComp->SetValueAsInt(Key.SelectedKeyName, Value);
 	}
 }
 
 void UBTFunctionLibrary::SetBlackboardValueAsFloat(UBTNode* NodeOwner, const struct FBlackboardKeySelector& Key, float Value)
 {
-	UBehaviorTreeComponent* BTComp = NodeOwner ? Cast<UBehaviorTreeComponent>(NodeOwner->GetOuter()) : NULL;
-	if (BTComp != NULL)
+	UBlackboardComponent* BlackboardComp = GetBlackboard(NodeOwner);
+	if (BlackboardComp != NULL)
 	{
-		BTComp->GetBlackboardComponent()->SetValueAsFloat(Key.SelectedKeyName, Value);
+		BlackboardComp->SetValueAsFloat(Key.SelectedKeyName, Value);
 	}
 }
 
 void UBTFunctionLibrary::SetBlackboardValueAsBool(UBTNode* NodeOwner, const struct FBlackboardKeySelector& Key, bool Value)
 {
-	UBehaviorTreeComponent* BTComp = NodeOwner ? Cast<UBehaviorTreeComponent>(NodeOwner->GetOuter()) : NULL;
-	if (BTComp != NULL)
+	UBlackboardComponent* BlackboardComp = GetBlackboard(NodeOwner);
+	if (BlackboardComp != NULL)
 	{
-		BTComp->GetBlackboardComponent()->SetValueAsBool(Key.SelectedKeyName, Value);
+		BlackboardComp->SetValueAsBool(Key.SelectedKeyName, Value);
 	}
 }
 
 void UBTFunctionLibrary::SetBlackboardValueAsString(UBTNode* NodeOwner, const struct FBlackboardKeySelector& Key, FString Value)
 {
-	UBehaviorTreeComponent* BTComp = NodeOwner ? Cast<UBehaviorTreeComponent>(NodeOwner->GetOuter()) : NULL;
-	if (BTComp != NULL)
+	UBlackboardComponent* BlackboardComp = GetBlackboard(NodeOwner);
+	if (BlackboardComp != NULL)
 	{
-		BTComp->GetBlackboardComponent()->SetValueAsString(Key.SelectedKeyName, Value);
+		BlackboardComp->SetValueAsString(Key.SelectedKeyName, Value);
 	}
 }
 
 void UBTFunctionLibrary::SetBlackboardValueAsName(UBTNode* NodeOwner, const struct FBlackboardKeySelector& Key, FName Value)
 {
-	UBehaviorTreeComponent* BTComp = NodeOwner ? Cast<UBehaviorTreeComponent>(NodeOwner->GetOuter()) : NULL;
-	if (BTComp != NULL)
+	UBlackboardComponent* BlackboardComp = GetBlackboard(NodeOwner);
+	if (BlackboardComp != NULL)
 	{
-		BTComp->GetBlackboardComponent()->SetValueAsName(Key.SelectedKeyName, Value);
+		BlackboardComp->SetValueAsName(Key.SelectedKeyName, Value);
 	}
 }
 
 void UBTFunctionLibrary::SetBlackboardValueAsVector(UBTNode* NodeOwner, const struct FBlackboardKeySelector& Key, FVector Value)
 {
-	UBehaviorTreeComponent* BTComp = NodeOwner ? Cast<UBehaviorTreeComponent>(NodeOwner->GetOuter()) : NULL;
-	if (BTComp != NULL)
+	UBlackboardComponent* BlackboardComp = GetBlackboard(NodeOwner);
+	if (BlackboardComp != NULL)
 	{
-		BTComp->GetBlackboardComponent()->SetValueAsVector(Key.SelectedKeyName, Value);
+		BlackboardComp->SetValueAsVector(Key.SelectedKeyName, Value);
 	}
 }
 
 void UBTFunctionLibrary::ClearBlackboardValueAsVector(UBTNode* NodeOwner, const struct FBlackboardKeySelector& Key)
 {
-	UBehaviorTreeComponent* BTComp = NodeOwner ? Cast<UBehaviorTreeComponent>(NodeOwner->GetOuter()) : NULL;
-	if (BTComp != NULL)
+	UBlackboardComponent* BlackboardComp = GetBlackboard(NodeOwner);
+	if (BlackboardComp != NULL)
 	{
-		BTComp->GetBlackboardComponent()->ClearValueAsVector(Key.SelectedKeyName);
+		BlackboardComp->ClearValueAsVector(Key.SelectedKeyName);
 	}
 }
 
diff --git a/Engine/Source/Runtime/AIModule/Private/BehaviorTree/BehaviorTreeComponent.cpp b/Engine/Source/Runtime/AIModule/Private/BehaviorTree/BehaviorTreeComponent.cpp
--- a/Engine/Source/Runtime/AIModule/Private/BehaviorTree/BehaviorTreeComponent.cpp
+++ b/Engine/Source/Runtime/AIModule/Private/BehaviorTree/BehaviorTreeComponent.cpp
@@ -1739,23 +1739,23 @@ void UBehaviorTreeComponent::UpdateDebuggerAfterExecution(const class UBTTaskNod
 #endif
 }
 
-void UBehaviorTreeComponent::StoreDebuggerBlackboard(TArray<FString>& BlackboardValueDesc) const
+void UBehaviorTreeComponent::StoreDebuggerBlackboard(TMap<FName, FString>& BlackboardValueDesc) const
 {
 #if USE_BEHAVIORTREE_DEBUGGER
 	if (BlackboardComp && BlackboardComp->HasValidAsset())
 	{
 		const int32 NumKeys = BlackboardComp->GetNumKeys();
-		BlackboardValueDesc.Reset();
-		BlackboardValueDesc.AddZeroed(NumKeys);
+		BlackboardValueDesc.Empty(NumKeys);
 
 		for (int32 KeyIndex = 0; KeyIndex < NumKeys; KeyIndex++)
 		{
-			BlackboardValueDesc[KeyIndex] = BlackboardComp->DescribeKeyValue(KeyIndex, EBlackboardDescription::OnlyValue);
-			
-			if (BlackboardValueDesc[KeyIndex].Len() == 0)
+			FString Value = BlackboardComp->DescribeKeyValue(KeyIndex, EBlackboardDescription::OnlyValue);
+			if (Value.Len() == 0)
 			{
-				BlackboardValueDesc[KeyIndex] = TEXT("n/a");
+				Value = TEXT("n/a");
 			}
+
+			BlackboardValueDesc.Add(BlackboardComp->GetKeyName(KeyIndex), Value);
 		}
 	}
 #endif
diff --git a/Engine/Source/Runtime/AIModule/Private/BehaviorTree/BlackboardComponent.cpp b/Engine/Source/Runtime/AIModule/Private/BehaviorTree/BlackboardComponent.cpp
--- a/Engine/Source/Runtime/AIModule/Private/BehaviorTree/BlackboardComponent.cpp
+++ b/Engine/Source/Runtime/AIModule/Private/BehaviorTree/BlackboardComponent.cpp
@@ -54,7 +54,7 @@ struct FBlackboardInitializationData
 	};
 };
 
-void UBlackboardComponent::InitializeBlackboard(class UBlackboardData* NewAsset)
+void UBlackboardComponent::InitializeBlackboard(UBlackboardData* NewAsset)
 {
 	// if we re-initialize with the same asset then there's no point
 	// in reseting, since we'd loose all the accumulated knowledge
@@ -116,12 +116,12 @@ void UBlackboardComponent::InitializeBlackboard(class UBlackboardData* NewAsset)
 	}
 }
 
-class UBrainComponent* UBlackboardComponent::GetBrainComponent() const
+UBrainComponent* UBlackboardComponent::GetBrainComponent() const
 {
 	return BrainComp;
 }
 
-class UBlackboardData* UBlackboardComponent::GetBlackboardAsset() const
+UBlackboardData* UBlackboardComponent::GetBlackboardAsset() const
 {
 	return BlackboardAsset;
 }
@@ -256,6 +256,11 @@ FString UBlackboardComponent::GetDebugInfoString(EBlackboardDescription::Type Mo
 	return DebugString;
 }
 
+FString UBlackboardComponent::DescribeKeyValue(const FName& KeyName, EBlackboardDescription::Type Mode) const
+{
+	return DescribeKeyValue(GetKeyID(KeyName), Mode);
+}
+
 FString UBlackboardComponent::DescribeKeyValue(FBlackboard::FKey KeyID, EBlackboardDescription::Type Mode) const
 {
 	FString Description;
