Added new node buttons to behavior tree editor.

Added 3 new buttons/combo buttons used to create new Blueprint-based node classes. Buttons turn into combo buttons if there is more than one type of base class available.
New class Blueprints are created alongside the current behavior tree asset.
Fixed toolbar buttons having no visibility control if they had no bound command list.

TTP# 337671 - Behavior Tree Editor: Extend toolbar with ôcreate new taskö, ôcreate new decoratorö, ôcreate new serviceö options

diff --git a/Engine/Source/Editor/BehaviorTreeEditor/BehaviorTreeEditor.Build.cs b/Engine/Source/Editor/BehaviorTreeEditor/BehaviorTreeEditor.Build.cs
--- a/Engine/Source/Editor/BehaviorTreeEditor/BehaviorTreeEditor.Build.cs
+++ b/Engine/Source/Editor/BehaviorTreeEditor/BehaviorTreeEditor.Build.cs
@@ -52,6 +52,7 @@ public class BehaviorTreeEditor : ModuleRules
                 "WorkspaceMenuStructure",
                 "PropertyEditor",
 				"AssetTools",
+				"AssetRegistry"
             }
 		);
 	}
diff --git a/Engine/Source/Editor/BehaviorTreeEditor/Private/BehaviorTreeEditor.cpp b/Engine/Source/Editor/BehaviorTreeEditor/Private/BehaviorTreeEditor.cpp
--- a/Engine/Source/Editor/BehaviorTreeEditor/Private/BehaviorTreeEditor.cpp
+++ b/Engine/Source/Editor/BehaviorTreeEditor/Private/BehaviorTreeEditor.cpp
@@ -8,6 +8,9 @@
 
 #include "BehaviorTree/Composites/BTComposite_SimpleParallel.h"
 #include "BehaviorTree/Tasks/BTTask_RunBehavior.h"
+#include "BehaviorTree/Tasks/BTTask_BlueprintBase.h"
+#include "BehaviorTree/Decorators/BTDecorator_BlueprintBase.h"
+#include "BehaviorTree/Services/BTService_BlueprintBase.h"
 #include "BehaviorTree/BTDecorator.h"
 
 #include "Toolkits/IToolkitHost.h"
@@ -25,6 +28,10 @@
 #include "DetailCustomizations/BlackboardDataDetails.h"
 #include "SBehaviorTreeBlackboardView.h"
 #include "SBehaviorTreeBlackboardEditor.h"
+#include "ClassViewerFilter.h"
+#include "ClassViewerModule.h"
+#include "AssetRegistryModule.h"
+#include "AssetToolsModule.h"
 
 #define LOCTEXT_NAMESPACE "BehaviorTreeEditor"
 
@@ -64,6 +71,7 @@ FBehaviorTreeEditor::~FBehaviorTreeEditor()
 
 	UPackage::PackageSavedEvent.RemoveRaw(this, &FBehaviorTreeEditor::OnPackageSaved);
 	FClassBrowseHelper::OnPackageListUpdated.RemoveRaw(this, &FBehaviorTreeEditor::OnClassListUpdated);
+
 	Debugger.Reset();
 }
 
@@ -1964,4 +1972,151 @@ void FBehaviorTreeEditor::RefreshDebugger()
 	Debugger->Refresh();
 }
 
+bool FBehaviorTreeEditor::CanCreateNewTask() const
+{
+	return !IsDebuggerReady();
+}
+
+bool FBehaviorTreeEditor::CanCreateNewDecorator() const
+{
+	return !IsDebuggerReady();
+}
+
+bool FBehaviorTreeEditor::CanCreateNewService() const
+{
+	return !IsDebuggerReady();
+}
+
+template <typename Type>
+class FNewNodeClassFilter : public IClassViewerFilter
+{
+public:
+	virtual bool IsClassAllowed(const FClassViewerInitializationOptions& InInitOptions, const UClass* InClass, TSharedRef< FClassViewerFilterFuncs > InFilterFuncs ) override
+	{
+		if(InClass != nullptr)
+		{
+			return InClass->IsChildOf(Type::StaticClass());
+		}
+		return false;
+	}
+
+	virtual bool IsUnloadedClassAllowed(const FClassViewerInitializationOptions& InInitOptions, const TSharedRef< const IUnloadedBlueprintData > InUnloadedClassData, TSharedRef< FClassViewerFilterFuncs > InFilterFuncs) override
+	{
+		return InUnloadedClassData->IsChildOf(Type::StaticClass());
+	}
+};
+
+
+TSharedRef<SWidget> FBehaviorTreeEditor::HandleCreateNewTaskMenu() const
+{
+	FClassViewerInitializationOptions Options;
+	Options.bShowUnloadedBlueprints = true;
+	Options.ClassFilter = MakeShareable( new FNewNodeClassFilter<UBTTask_BlueprintBase> );
+
+	FOnClassPicked OnPicked( FOnClassPicked::CreateSP( this, &FBehaviorTreeEditor::HandleNewNodeClassPicked ) );
+
+	return FModuleManager::LoadModuleChecked<FClassViewerModule>("ClassViewer").CreateClassViewer(Options, OnPicked);
+}
+
+TSharedRef<SWidget> FBehaviorTreeEditor::HandleCreateNewDecoratorMenu() const
+{
+	FClassViewerInitializationOptions Options;
+	Options.bShowUnloadedBlueprints = true;
+	Options.ClassFilter = MakeShareable( new FNewNodeClassFilter<UBTDecorator_BlueprintBase> );
+
+	FOnClassPicked OnPicked( FOnClassPicked::CreateSP( this, &FBehaviorTreeEditor::HandleNewNodeClassPicked ) );
+
+	return FModuleManager::LoadModuleChecked<FClassViewerModule>("ClassViewer").CreateClassViewer(Options, OnPicked);
+}
+
+TSharedRef<SWidget> FBehaviorTreeEditor::HandleCreateNewServiceMenu() const
+{
+	FClassViewerInitializationOptions Options;
+	Options.bShowUnloadedBlueprints = true;
+	Options.ClassFilter = MakeShareable( new FNewNodeClassFilter<UBTService_BlueprintBase> );
+
+	FOnClassPicked OnPicked( FOnClassPicked::CreateSP( this, &FBehaviorTreeEditor::HandleNewNodeClassPicked ) );
+
+	return FModuleManager::LoadModuleChecked<FClassViewerModule>("ClassViewer").CreateClassViewer(Options, OnPicked);
+}
+
+void FBehaviorTreeEditor::HandleNewNodeClassPicked(UClass* InClass) const
+{
+	if(BehaviorTree != nullptr)
+	{
+		FString ClassName = InClass->GetName();
+		ClassName.RemoveFromEnd(TEXT("_C"));
+
+		FString PathName = BehaviorTree->GetOutermost()->GetPathName();
+		PathName = FPaths::GetPath(PathName);
+		PathName /= ClassName;
+
+		FString Name;
+		FString PackageName;
+		FAssetToolsModule& AssetToolsModule = FModuleManager::GetModuleChecked<FAssetToolsModule>("AssetTools");
+		AssetToolsModule.Get().CreateUniqueAssetName(PathName, TEXT("_New"), PackageName, Name);
+
+		UPackage* Package = CreatePackage(NULL, *PackageName);
+		if (ensure(Package))
+		{
+			// Create and init a new Blueprint
+			if (UBlueprint* NewBP = FKismetEditorUtilities::CreateBlueprint(InClass, Package, FName(*Name), BPTYPE_Normal, UBlueprint::StaticClass(), UBlueprintGeneratedClass::StaticClass()))
+			{
+				FAssetEditorManager::Get().OpenEditorForAsset(NewBP);
+
+				// Notify the asset registry
+				FAssetRegistryModule::AssetCreated(NewBP);
+
+				// Mark the package dirty...
+				Package->MarkPackageDirty();
+			}
+		}
+	}
+}
+
+void FBehaviorTreeEditor::CreateNewTask() const
+{
+	HandleNewNodeClassPicked(UBTTask_BlueprintBase::StaticClass());
+}
+
+bool FBehaviorTreeEditor::IsNewTaskButtonVisible() const
+{
+	return !FClassBrowseHelper::IsMoreThanOneTaskClassAvailable();
+}
+
+bool FBehaviorTreeEditor::IsNewTaskComboVisible() const
+{
+	return FClassBrowseHelper::IsMoreThanOneTaskClassAvailable();
+}
+
+void FBehaviorTreeEditor::CreateNewDecorator() const
+{
+	HandleNewNodeClassPicked(UBTDecorator_BlueprintBase::StaticClass());
+}
+
+bool FBehaviorTreeEditor::IsNewDecoratorButtonVisible() const
+{
+	return !FClassBrowseHelper::IsMoreThanOneDecoratorClassAvailable();
+}
+
+bool FBehaviorTreeEditor::IsNewDecoratorComboVisible() const
+{
+	return FClassBrowseHelper::IsMoreThanOneDecoratorClassAvailable();
+}
+
+void FBehaviorTreeEditor::CreateNewService() const
+{
+	HandleNewNodeClassPicked(UBTService_BlueprintBase::StaticClass());
+}
+
+bool FBehaviorTreeEditor::IsNewServiceButtonVisible() const
+{
+	return !FClassBrowseHelper::IsMoreThanOneServiceClassAvailable();
+}
+
+bool FBehaviorTreeEditor::IsNewServiceComboVisible() const
+{
+	return FClassBrowseHelper::IsMoreThanOneServiceClassAvailable();
+}
+
 #undef LOCTEXT_NAMESPACE
diff --git a/Engine/Source/Editor/BehaviorTreeEditor/Private/BehaviorTreeEditor.h b/Engine/Source/Editor/BehaviorTreeEditor/Private/BehaviorTreeEditor.h
--- a/Engine/Source/Editor/BehaviorTreeEditor/Private/BehaviorTreeEditor.h
+++ b/Engine/Source/Editor/BehaviorTreeEditor/Private/BehaviorTreeEditor.h
@@ -184,6 +184,54 @@ public:
 	/** Get the currently selected blackboard entry */
 	void GetBlackboardSelectionInfo(int32& OutSelectionIndex, bool& bOutIsInherited) const;
 
+	/** Check to see if we can create a new task node */
+	bool CanCreateNewTask() const;
+
+	/** Check to see if we can create a new decorator node */
+	bool CanCreateNewDecorator() const;
+
+	/** Check to see if we can create a new service node */
+	bool CanCreateNewService() const;
+
+	/** Create the menu used to make a new task node */
+	TSharedRef<SWidget> HandleCreateNewTaskMenu() const;
+
+	/** Create the menu used to make a new decorator */
+	TSharedRef<SWidget> HandleCreateNewDecoratorMenu() const;
+
+	/** Create the menu used to make a new service */
+	TSharedRef<SWidget> HandleCreateNewServiceMenu() const;
+
+	/** Handler for when a node class is picked */
+	void HandleNewNodeClassPicked(UClass* InClass) const;
+
+	/** Create a new task from UBTTask_BlueprintBase */
+	void CreateNewTask() const;
+
+	/** Whether the single button to create a new Blueprint-based task is visible */
+	bool IsNewTaskButtonVisible() const;
+
+	/** Whether the combo button to create a new Blueprint-based task from all available base classes is visible */
+	bool IsNewTaskComboVisible() const;
+
+	/** Create a new decorator from UBTDecorator_BlueprintBase */
+	void CreateNewDecorator() const;
+
+	/** Whether the single button to create a new Blueprint-based decorator is visible */
+	bool IsNewDecoratorButtonVisible() const;
+
+	/** Whether the combo button to create a new Blueprint-based decorator from all available base classes is visible */
+	bool IsNewDecoratorComboVisible() const;
+
+	/** Create a new service from UBTService_BlueprintBase */
+	void CreateNewService() const;
+
+	/** Whether the single button to create a new Blueprint-based service is visible */
+	bool IsNewServiceButtonVisible() const;
+
+	/** Whether the combo button to create a new Blueprint-based service from all available base classes is visible */
+	bool IsNewServiceComboVisible() const;
+
 protected:
 	/** Called when "Save" is clicked for this asset */
 	virtual void SaveAsset_Execute() override;
diff --git a/Engine/Source/Editor/BehaviorTreeEditor/Private/BehaviorTreeEditorModes.cpp b/Engine/Source/Editor/BehaviorTreeEditor/Private/BehaviorTreeEditorModes.cpp
--- a/Engine/Source/Editor/BehaviorTreeEditor/Private/BehaviorTreeEditorModes.cpp
+++ b/Engine/Source/Editor/BehaviorTreeEditor/Private/BehaviorTreeEditorModes.cpp
@@ -58,6 +58,7 @@ FBehaviorTreeEditorApplicationMode::FBehaviorTreeEditorApplicationMode(TSharedPt
 	
 	InBehaviorTreeEditor->GetToolbarBuilder()->AddModesToolbar(ToolbarExtender);
 	InBehaviorTreeEditor->GetToolbarBuilder()->AddDebuggerToolbar(ToolbarExtender);
+	InBehaviorTreeEditor->GetToolbarBuilder()->AddBehaviorTreeToolbar(ToolbarExtender);
 }
 
 void FBehaviorTreeEditorApplicationMode::RegisterTabFactories(TSharedPtr<FTabManager> InTabManager)
diff --git a/Engine/Source/Editor/BehaviorTreeEditor/Private/BehaviorTreeEditorToolbar.cpp b/Engine/Source/Editor/BehaviorTreeEditor/Private/BehaviorTreeEditorToolbar.cpp
--- a/Engine/Source/Editor/BehaviorTreeEditor/Private/BehaviorTreeEditorToolbar.cpp
+++ b/Engine/Source/Editor/BehaviorTreeEditor/Private/BehaviorTreeEditorToolbar.cpp
@@ -95,6 +95,17 @@ void FBehaviorTreeEditorToolbar::AddDebuggerToolbar(TSharedPtr<FExtender> Extend
 	BehaviorTreeEditorPtr->AddToolbarExtender(ToolbarExtender);
 }
 
+
+void FBehaviorTreeEditorToolbar::AddBehaviorTreeToolbar(TSharedPtr<FExtender> Extender)
+{
+	check(BehaviorTreeEditor.IsValid());
+	TSharedPtr<FBehaviorTreeEditor> BehaviorTreeEditorPtr = BehaviorTreeEditor.Pin();
+
+	TSharedPtr<FExtender> ToolbarExtender = MakeShareable(new FExtender);
+	ToolbarExtender->AddToolBarExtension("Asset", EExtensionHook::After, BehaviorTreeEditorPtr->GetToolkitCommands(), FToolBarExtensionDelegate::CreateSP( this, &FBehaviorTreeEditorToolbar::FillBehaviorTreeToolbar ));
+	BehaviorTreeEditorPtr->AddToolbarExtender(ToolbarExtender);
+}
+
 void FBehaviorTreeEditorToolbar::FillModesToolbar(FToolBarBuilder& ToolbarBuilder)
 {
 	check(BehaviorTreeEditor.IsValid());
@@ -132,4 +143,107 @@ void FBehaviorTreeEditorToolbar::FillModesToolbar(FToolBarBuilder& ToolbarBuilde
 	BehaviorTreeEditorPtr->AddToolbarWidget(SNew(SSpacer).Size(FVector2D(4.0f, 1.0f)));
 }
 
-#undef LOCTEXT_NAMESPACE
+void FBehaviorTreeEditorToolbar::FillBehaviorTreeToolbar(FToolBarBuilder& ToolbarBuilder)
+{
+	check(BehaviorTreeEditor.IsValid());
+	TSharedPtr<FBehaviorTreeEditor> BehaviorTreeEditorPtr = BehaviorTreeEditor.Pin();
+
+	if (!BehaviorTreeEditorPtr->IsDebuggerReady() && BehaviorTreeEditorPtr->GetCurrentMode() == FBehaviorTreeEditor::BehaviorTreeMode)
+	{
+		ToolbarBuilder.BeginSection("BehaviorTree");
+		{
+			const FText NewTaskLabel = LOCTEXT( "NewTask_Label", "New Task" );
+			const FText NewTaskTooltip = LOCTEXT( "NewTask_ToolTip", "Create a new task node Blueprint from a base class" );
+			const FSlateIcon NewTaskIcon = FSlateIcon(FEditorStyle::GetStyleSetName(), "BTEditor.Graph.NewTask");
+
+			ToolbarBuilder.AddToolBarButton(
+				FUIAction(
+					FExecuteAction::CreateSP(BehaviorTreeEditorPtr.Get(), &FBehaviorTreeEditor::CreateNewTask),
+					FCanExecuteAction::CreateSP(BehaviorTreeEditorPtr.Get(), &FBehaviorTreeEditor::CanCreateNewTask),
+					FIsActionChecked(),
+					FIsActionButtonVisible::CreateSP(BehaviorTreeEditorPtr.Get(), &FBehaviorTreeEditor::IsNewTaskButtonVisible)
+					), 
+				NAME_None,
+				NewTaskLabel,
+				NewTaskTooltip,
+				NewTaskIcon
+			);
+
+			ToolbarBuilder.AddComboButton(
+				FUIAction(
+					FExecuteAction(),
+					FCanExecuteAction::CreateSP(BehaviorTreeEditorPtr.Get(), &FBehaviorTreeEditor::CanCreateNewTask),
+					FIsActionChecked(),
+					FIsActionButtonVisible::CreateSP(BehaviorTreeEditorPtr.Get(), &FBehaviorTreeEditor::IsNewTaskComboVisible)
+					), 
+				FOnGetContent::CreateSP(BehaviorTreeEditorPtr.Get(), &FBehaviorTreeEditor::HandleCreateNewTaskMenu),
+				NewTaskLabel,
+				NewTaskTooltip,
+				NewTaskIcon
+			);
+
+			const FText NewDecoratorLabel = LOCTEXT( "NewDecorator_Label", "New Decorator" );
+			const FText NewDecoratorTooltip = LOCTEXT( "NewDecorator_ToolTip", "Create a new decorator node Blueprint from a base class" );
+			const FSlateIcon NewDecoratorIcon = FSlateIcon(FEditorStyle::GetStyleSetName(), "BTEditor.Graph.NewDecorator");
+
+			ToolbarBuilder.AddToolBarButton(
+				FUIAction(
+					FExecuteAction::CreateSP(BehaviorTreeEditorPtr.Get(), &FBehaviorTreeEditor::CreateNewDecorator),
+					FCanExecuteAction::CreateSP(BehaviorTreeEditorPtr.Get(), &FBehaviorTreeEditor::CanCreateNewDecorator),
+					FIsActionChecked(),
+					FIsActionButtonVisible::CreateSP(BehaviorTreeEditorPtr.Get(), &FBehaviorTreeEditor::IsNewDecoratorButtonVisible)
+					), 
+				NAME_None,
+				NewDecoratorLabel,
+				NewDecoratorTooltip,
+				NewDecoratorIcon
+			);
+
+			ToolbarBuilder.AddComboButton(
+				FUIAction(
+					FExecuteAction(),
+					FCanExecuteAction::CreateSP(BehaviorTreeEditorPtr.Get(), &FBehaviorTreeEditor::CanCreateNewDecorator),
+					FIsActionChecked(),
+					FIsActionButtonVisible::CreateSP(BehaviorTreeEditorPtr.Get(), &FBehaviorTreeEditor::IsNewDecoratorComboVisible)
+					), 
+				FOnGetContent::CreateSP(BehaviorTreeEditorPtr.Get(), &FBehaviorTreeEditor::HandleCreateNewDecoratorMenu),
+				NewDecoratorLabel,
+				NewDecoratorTooltip,
+				NewDecoratorIcon
+			);
+
+			const FText NewServiceLabel = LOCTEXT( "NewService_Label", "New Service" );
+			const FText NewServiceTooltip = LOCTEXT( "NewService_ToolTip", "Create a new service node Blueprint from a base class" );
+			const FSlateIcon NewServiceIcon = FSlateIcon(FEditorStyle::GetStyleSetName(), "BTEditor.Graph.NewService");
+
+			ToolbarBuilder.AddToolBarButton(
+				FUIAction(
+					FExecuteAction::CreateSP(BehaviorTreeEditorPtr.Get(), &FBehaviorTreeEditor::CreateNewService),
+					FCanExecuteAction::CreateSP(BehaviorTreeEditorPtr.Get(), &FBehaviorTreeEditor::CanCreateNewService),
+					FIsActionChecked(),
+					FIsActionButtonVisible::CreateSP(BehaviorTreeEditorPtr.Get(), &FBehaviorTreeEditor::IsNewServiceButtonVisible)
+					), 
+				NAME_None,
+				NewServiceLabel,
+				NewServiceTooltip,
+				NewServiceIcon
+			);
+
+			ToolbarBuilder.AddComboButton(
+				FUIAction(
+					FExecuteAction(),
+					FCanExecuteAction::CreateSP(BehaviorTreeEditorPtr.Get(), &FBehaviorTreeEditor::CanCreateNewService),
+					FIsActionChecked(),
+					FIsActionButtonVisible::CreateSP(BehaviorTreeEditorPtr.Get(), &FBehaviorTreeEditor::IsNewServiceComboVisible)
+					), 
+				FOnGetContent::CreateSP(BehaviorTreeEditorPtr.Get(), &FBehaviorTreeEditor::HandleCreateNewServiceMenu),
+				NewServiceLabel,
+				NewServiceTooltip,
+				NewServiceIcon
+			);
+		}
+		ToolbarBuilder.EndSection();
+	}
+}
+
+#undef LOCTEXT_NAMESPACE
\ No newline at end of file
diff --git a/Engine/Source/Editor/BehaviorTreeEditor/Private/BehaviorTreeEditorToolbar.h b/Engine/Source/Editor/BehaviorTreeEditor/Private/BehaviorTreeEditorToolbar.h
--- a/Engine/Source/Editor/BehaviorTreeEditor/Private/BehaviorTreeEditorToolbar.h
+++ b/Engine/Source/Editor/BehaviorTreeEditor/Private/BehaviorTreeEditorToolbar.h
@@ -10,10 +10,12 @@ public:
 
 	void AddModesToolbar(TSharedPtr<FExtender> Extender);
 	void AddDebuggerToolbar(TSharedPtr<FExtender> Extender);
+	void AddBehaviorTreeToolbar(TSharedPtr<FExtender> Extender);
 
 private:
 	void FillModesToolbar(FToolBarBuilder& ToolbarBuilder);
 	void FillDebuggerToolbar(FToolBarBuilder& ToolbarBuilder);
+	void FillBehaviorTreeToolbar(FToolBarBuilder& ToolbarBuilder);
 
 protected:
 	/** Pointer back to the blueprint editor tool that owns us */
diff --git a/Engine/Source/Editor/EditorStyle/Private/SlateEditorStyle.cpp b/Engine/Source/Editor/EditorStyle/Private/SlateEditorStyle.cpp
--- a/Engine/Source/Editor/EditorStyle/Private/SlateEditorStyle.cpp
+++ b/Engine/Source/Editor/EditorStyle/Private/SlateEditorStyle.cpp
@@ -2870,6 +2870,13 @@ void FSlateEditorStyle::FStyle::SetupGraphEditorStyles()
 			Set( "ClassIcon.BlackboardKeyType_Rotator", new IMAGE_BRUSH( "Icons/pill_16x", Icon16x16, FLinearColor(0.353393f, 0.454175f, 1.0f, 1.0f) ) );
 			Set( "ClassIcon.BlackboardKeyType_String", new IMAGE_BRUSH( "Icons/pill_16x", Icon16x16, FLinearColor(1.0f, 0.0f, 0.660537f, 1.0f) ) );
 			Set( "ClassIcon.BlackboardKeyType_Vector", new IMAGE_BRUSH( "Icons/pill_16x", Icon16x16, FLinearColor(1.0f, 0.591255f, 0.016512f, 1.0f) ) );
+
+			Set( "BTEditor.Graph.NewTask", new IMAGE_BRUSH( "BehaviorTree/NewTask_40x", FVector2D(40,40)));
+			Set( "BTEditor.Graph.NewTask.Small", new IMAGE_BRUSH( "BehaviorTree/NewTask_40x", FVector2D(16,16)));
+			Set( "BTEditor.Graph.NewDecorator", new IMAGE_BRUSH( "BehaviorTree/NewDecorator_40x", FVector2D(40,40)));
+			Set( "BTEditor.Graph.NewDecorator.Small", new IMAGE_BRUSH( "BehaviorTree/NewDecorator_40x", FVector2D(16,16)));
+			Set( "BTEditor.Graph.NewService", new IMAGE_BRUSH( "BehaviorTree/NewService_40x", FVector2D(40,40)));
+			Set( "BTEditor.Graph.NewService.Small", new IMAGE_BRUSH( "BehaviorTree/NewService_40x", FVector2D(16,16)));
 		}
 
 		// Visible on hover button for transition node
diff --git a/Engine/Source/Runtime/Slate/Private/Framework/MultiBox/SToolBarButtonBlock.cpp b/Engine/Source/Runtime/Slate/Private/Framework/MultiBox/SToolBarButtonBlock.cpp
--- a/Engine/Source/Runtime/Slate/Private/Framework/MultiBox/SToolBarButtonBlock.cpp
+++ b/Engine/Source/Runtime/Slate/Private/Framework/MultiBox/SToolBarButtonBlock.cpp
@@ -378,10 +378,15 @@ bool SToolBarButtonBlock::IsEnabled() const
 EVisibility SToolBarButtonBlock::GetBlockVisibility() const
 {
 	TSharedPtr< const FUICommandList > ActionList = MultiBlock->GetActionList();
+	const FUIAction& DirectActions = MultiBlock->GetDirectActions();
 	if( ActionList.IsValid() )
 	{
 		return ActionList->GetVisibility( MultiBlock->GetAction().ToSharedRef() );
 	}
+	else if(DirectActions.IsActionVisibleDelegate.IsBound())
+	{
+		return DirectActions.IsActionVisibleDelegate.Execute() ? EVisibility::Visible : EVisibility::Collapsed;
+	}
 
 	return EVisibility::Visible;
 }
