CIS fix - missing changes.

diff --git a/Engine/Source/Editor/BehaviorTreeEditor/Classes/BehaviorTreeEditorTypes.h b/Engine/Source/Editor/BehaviorTreeEditor/Classes/BehaviorTreeEditorTypes.h
--- a/Engine/Source/Editor/BehaviorTreeEditor/Classes/BehaviorTreeEditorTypes.h
+++ b/Engine/Source/Editor/BehaviorTreeEditor/Classes/BehaviorTreeEditorTypes.h
@@ -90,6 +90,15 @@ struct FClassBrowseHelper
 	static bool IsClassKnown(const FClassData& ClassData);
 	static FOnPackageListUpdated OnPackageListUpdated;
 
+	/** Whether we have more than one task class available */
+	static bool IsMoreThanOneTaskClassAvailable();
+
+	/** Whether we have more than one decorator class available */
+	static bool IsMoreThanOneDecoratorClassAvailable();
+
+	/** Whether we have more than one service class available */
+	static bool IsMoreThanOneServiceClassAvailable();
+
 private:
 
 	TSharedPtr<FClassDataNode> RootNode;
@@ -105,6 +114,18 @@ private:
 
 	bool IsHidingParentClass(UClass* Class);
 	bool IsPackageSaved(FName PackageName);
+
+	/** Update the flags that determine whether to allow creation of new nodes from multiple base classes */
+	void UpdateAvailableNodeClasses();
+
+	/** Cached flag to determine whether we have more than one task class available */
+	static bool bMoreThanOneTaskClassAvailable;
+
+	/** Cached flag to determine whether we have more than one decorator class available */
+	static bool bMoreThanOneDecoratorClassAvailable;
+
+	/** Cached flag to determine whether we have more than one service class available */
+	static bool bMoreThanOneServiceClassAvailable;
 };
 
 struct FCompareNodeXLocation
diff --git a/Engine/Source/Editor/BehaviorTreeEditor/Private/BehaviorTreeEditorTypes.cpp b/Engine/Source/Editor/BehaviorTreeEditor/Private/BehaviorTreeEditorTypes.cpp
--- a/Engine/Source/Editor/BehaviorTreeEditor/Private/BehaviorTreeEditorTypes.cpp
+++ b/Engine/Source/Editor/BehaviorTreeEditor/Private/BehaviorTreeEditorTypes.cpp
@@ -90,6 +90,9 @@ UClass* FClassData::GetClass(bool bSilent)
 //////////////////////////////////////////////////////////////////////////
 FClassBrowseHelper::FOnPackageListUpdated FClassBrowseHelper::OnPackageListUpdated;
 TArray<FName> FClassBrowseHelper::UnknownPackages;
+bool FClassBrowseHelper::bMoreThanOneTaskClassAvailable = false;
+bool FClassBrowseHelper::bMoreThanOneDecoratorClassAvailable = false;
+bool FClassBrowseHelper::bMoreThanOneServiceClassAvailable = false;
 
 void FClassDataNode::AddUniqueSubNode(TSharedPtr<FClassDataNode> SubNode)
 {
@@ -116,6 +119,8 @@ FClassBrowseHelper::FClassBrowseHelper()
 	GEditor->OnHotReload().AddRaw( this, &FClassBrowseHelper::InvalidateCache );
 	GEditor->OnBlueprintCompiled().AddRaw( this, &FClassBrowseHelper::InvalidateCache );
 	GEditor->OnClassPackageLoadedOrUnloaded().AddRaw( this, &FClassBrowseHelper::InvalidateCache );
+
+	UpdateAvailableNodeClasses();
 }
 
 FClassBrowseHelper::~FClassBrowseHelper()
@@ -226,6 +231,12 @@ void FClassBrowseHelper::OnAssetAdded(const class FAssetData& AssetData)
 		ParentNode->AddUniqueSubNode(Node);
 		Node->ParentNode = ParentNode;
 	}
+
+	FAssetRegistryModule& AssetRegistryModule = FModuleManager::LoadModuleChecked<FAssetRegistryModule>(TEXT("AssetRegistry"));
+	if(!AssetRegistryModule.Get().IsLoadingAssets())
+	{
+		UpdateAvailableNodeClasses();
+	}
 }
 
 void FClassBrowseHelper::OnAssetRemoved(const class FAssetData& AssetData)
@@ -243,11 +254,19 @@ void FClassBrowseHelper::OnAssetRemoved(const class FAssetData& AssetData)
 			Node->ParentNode->SubNodes.RemoveSingleSwap(Node);
 		}
 	}
+
+	FAssetRegistryModule& AssetRegistryModule = FModuleManager::LoadModuleChecked<FAssetRegistryModule>(TEXT("AssetRegistry"));
+	if(!AssetRegistryModule.Get().IsLoadingAssets())
+	{
+		UpdateAvailableNodeClasses();
+	}
 }
 
 void FClassBrowseHelper::InvalidateCache()
 {
 	RootNode.Reset();
+
+	UpdateAvailableNodeClasses();
 }
 
 TSharedPtr<FClassDataNode> FClassBrowseHelper::CreateClassDataNode(const class FAssetData& AssetData)
@@ -423,6 +442,54 @@ void FClassBrowseHelper::AddClassGraphChildren(TSharedPtr<FClassDataNode> Node,
 	}
 }
 
+bool FClassBrowseHelper::IsMoreThanOneTaskClassAvailable()
+{
+	return bMoreThanOneTaskClassAvailable;
+}
+
+bool FClassBrowseHelper::IsMoreThanOneDecoratorClassAvailable()
+{
+	return bMoreThanOneDecoratorClassAvailable;
+}
+
+bool FClassBrowseHelper::IsMoreThanOneServiceClassAvailable()
+{
+	return bMoreThanOneServiceClassAvailable;
+}
+
+void FClassBrowseHelper::UpdateAvailableNodeClasses()
+{
+	if(FModuleManager::Get().IsModuleLoaded(TEXT("AssetRegistry")))
+	{
+		FAssetRegistryModule& AssetRegistryModule = FModuleManager::GetModuleChecked<FAssetRegistryModule>(TEXT("AssetRegistry"));
+
+		const bool bSearchSubClasses = true;
+		{
+			TArray<FName> ClassNames;
+			ClassNames.Add(UBTTask_BlueprintBase::StaticClass()->GetFName());
+			TSet<FName> DerivedClassNames;
+			AssetRegistryModule.Get().GetDerivedClassNames(ClassNames, TSet<FName>(), DerivedClassNames);
+			bMoreThanOneTaskClassAvailable = DerivedClassNames.Num() > 1;
+		}
+
+		{
+			TArray<FName> ClassNames;
+			ClassNames.Add(UBTDecorator_BlueprintBase::StaticClass()->GetFName());
+			TSet<FName> DerivedClassNames;
+			AssetRegistryModule.Get().GetDerivedClassNames(ClassNames, TSet<FName>(), DerivedClassNames);
+			bMoreThanOneDecoratorClassAvailable = DerivedClassNames.Num() > 1;
+		}
+
+		{
+			TArray<FName> ClassNames;
+			ClassNames.Add(UBTService_BlueprintBase::StaticClass()->GetFName());
+			TSet<FName> DerivedClassNames;
+			AssetRegistryModule.Get().GetDerivedClassNames(ClassNames, TSet<FName>(), DerivedClassNames);
+			bMoreThanOneServiceClassAvailable = DerivedClassNames.Num() > 1;
+		}
+	}
+}
+
 //////////////////////////////////////////////////////////////////////////
 
 UBehaviorTreeEditorTypes::UBehaviorTreeEditorTypes(const class FPostConstructInitializeProperties& PCIP) : Super(PCIP)
