Behavior Tree root nodes now have correct tooltip

Root nodes do not initialize their class data member (as they are never placed, only generated) which is to generate the default tooltip. Overriding the function to provide a useful tooltip for nodes.

TTP# 339280 - UE4: BEHAVIORTREE: Class not found tooltop still appears after saving

diff --git a/Engine/Source/Editor/BehaviorTreeEditor/Classes/BehaviorTreeGraphNode_Root.h b/Engine/Source/Editor/BehaviorTreeEditor/Classes/BehaviorTreeGraphNode_Root.h
--- a/Engine/Source/Editor/BehaviorTreeEditor/Classes/BehaviorTreeGraphNode_Root.h
+++ b/Engine/Source/Editor/BehaviorTreeEditor/Classes/BehaviorTreeGraphNode_Root.h
@@ -20,6 +20,7 @@ class UBehaviorTreeGraphNode_Root : public UBehaviorTreeGraphNode
 	virtual FText GetNodeTitle(ENodeTitleType::Type TitleType) const override;
 	/** gets icon resource name for title bar */
 	virtual FName GetNameIcon() const override;
+	virtual FString GetTooltip() const override;
 
 	virtual void PostEditChangeProperty( struct FPropertyChangedEvent& PropertyChangedEvent) override;
 	virtual FString	GetDescription() const override;
diff --git a/Engine/Source/Editor/BehaviorTreeEditor/Private/BehaviorTreeGraphNode_Root.cpp b/Engine/Source/Editor/BehaviorTreeEditor/Private/BehaviorTreeGraphNode_Root.cpp
--- a/Engine/Source/Editor/BehaviorTreeEditor/Private/BehaviorTreeGraphNode_Root.cpp
+++ b/Engine/Source/Editor/BehaviorTreeEditor/Private/BehaviorTreeGraphNode_Root.cpp
@@ -40,6 +40,11 @@ FName UBehaviorTreeGraphNode_Root::GetNameIcon() const
 	return FName("BTEditor.Graph.BTNode.Root.Icon");
 }
 
+FString UBehaviorTreeGraphNode_Root::GetTooltip() const
+{
+	return NSLOCTEXT("BehaviorTreeGraphNode", "RootTooltip", "Root node of this behavior tree, holds Blackboard data").ToString();
+}
+
 void UBehaviorTreeGraphNode_Root::PostEditChangeProperty( struct FPropertyChangedEvent& PropertyChangedEvent)
 {
 	Super::PostEditChangeProperty(PropertyChangedEvent);
