CIS fix - non-unity.

diff --git a/Engine/Source/Editor/BehaviorTreeEditor/Private/BehaviorTreeEditorTypes.cpp b/Engine/Source/Editor/BehaviorTreeEditor/Private/BehaviorTreeEditorTypes.cpp
--- a/Engine/Source/Editor/BehaviorTreeEditor/Private/BehaviorTreeEditorTypes.cpp
+++ b/Engine/Source/Editor/BehaviorTreeEditor/Private/BehaviorTreeEditorTypes.cpp
@@ -5,6 +5,9 @@
 #include "PackageTools.h"
 #include "MessageLog.h"
 #include "AssetRegistryModule.h"
+#include "BehaviorTree/Tasks/BTTask_BlueprintBase.h"
+#include "BehaviorTree/Decorators/BTDecorator_BlueprintBase.h"
+#include "BehaviorTree/Services/BTService_BlueprintBase.h"
 
 const FString UBehaviorTreeEditorTypes::PinCategory_MultipleNodes("MultipleNodes");
 const FString UBehaviorTreeEditorTypes::PinCategory_SingleComposite("SingleComposite");
