Fixed crash when deleting then attempting to rename blackboard key

Made sure the details panel gets a refresh on delete.

TTP# 339846 - CRITICAL: DATA ASSET: CRASH: Occurs after deleting and modifying the name of a key

diff --git a/Engine/Source/Editor/BehaviorTreeEditor/Private/SBehaviorTreeBlackboardEditor.cpp b/Engine/Source/Editor/BehaviorTreeEditor/Private/SBehaviorTreeBlackboardEditor.cpp
--- a/Engine/Source/Editor/BehaviorTreeEditor/Private/SBehaviorTreeBlackboardEditor.cpp
+++ b/Engine/Source/Editor/BehaviorTreeEditor/Private/SBehaviorTreeBlackboardEditor.cpp
@@ -102,6 +102,12 @@ void SBehaviorTreeBlackboardEditor::HandleDeleteEntry()
 			}
 
 			GraphActionMenu->RefreshAllActions(true);
+
+			// signal de-selection
+			if(OnEntrySelected.IsBound())
+			{
+				OnEntrySelected.Execute(nullptr, false);
+			}
 		}
 	}
 }
