Prevented behavior tree nodes from going pink after undo

Made sure node owner is up-to-date after undo. This fixes the outer of the node being incorrect post-undo. Nodes are renamed away non-transactionally when deleted and cleaned up, so the change to the outer gets ignored and is not restored normally. This ensures that it gets properly fixed up.

TTP# 339299 - UE4: BEHAVIORTREE: Node turns pink from Undoing a delete

diff --git a/Engine/Source/Editor/BehaviorTreeEditor/Classes/BehaviorTreeGraphNode.h b/Engine/Source/Editor/BehaviorTreeEditor/Classes/BehaviorTreeGraphNode.h
--- a/Engine/Source/Editor/BehaviorTreeEditor/Classes/BehaviorTreeGraphNode.h
+++ b/Engine/Source/Editor/BehaviorTreeEditor/Classes/BehaviorTreeGraphNode.h
@@ -41,7 +41,10 @@ class UBehaviorTreeGraphNode : public UEdGraphNode
 	// End UEdGraphNode Interface
 
 	// Begin UObject Interface
+#if WITH_EDITOR
 	virtual void PostEditImport() override;
+	virtual void PostEditUndo() override;
+#endif
 	// End UObject
 
 	// @return the input pin for this state
diff --git a/Engine/Source/Editor/BehaviorTreeEditor/Private/BehaviorTreeGraphNode.cpp b/Engine/Source/Editor/BehaviorTreeEditor/Private/BehaviorTreeGraphNode.cpp
--- a/Engine/Source/Editor/BehaviorTreeEditor/Private/BehaviorTreeGraphNode.cpp
+++ b/Engine/Source/Editor/BehaviorTreeEditor/Private/BehaviorTreeGraphNode.cpp
@@ -73,6 +73,7 @@ void UBehaviorTreeGraphNode::PrepareForCopying()
 		NodeInstance->Rename(NULL, this, REN_DontCreateRedirectors | REN_DoNotDirty );
 	}
 }
+#if WITH_EDITOR
 
 void UBehaviorTreeGraphNode::PostEditImport()
 {
@@ -87,6 +88,13 @@ void UBehaviorTreeGraphNode::PostEditImport()
 	}
 }
 
+void UBehaviorTreeGraphNode::PostEditUndo()
+{
+	ResetNodeOwner();
+}
+
+#endif
+
 void UBehaviorTreeGraphNode::PostCopyNode()
 {
 	ResetNodeOwner();
