Behavior Tree nodes are now categorized

Added new BlueprintCategory property to UBlueprints that propagates to class metadata on Blueprint compilation. This metadata is then used for C++ and Blueprint-based node classes when displayed in the behavior tree palette.
Also fixed blueprint decorators incorrectly displaying no icon.

TTP# 337865 - [EDITOR][BT editor] BT nodes could use a way to define custom categories

diff --git a/Engine/Source/Editor/BehaviorTreeEditor/Classes/BehaviorTreeEditorTypes.h b/Engine/Source/Editor/BehaviorTreeEditor/Classes/BehaviorTreeEditorTypes.h
--- a/Engine/Source/Editor/BehaviorTreeEditor/Classes/BehaviorTreeEditorTypes.h
+++ b/Engine/Source/Editor/BehaviorTreeEditor/Classes/BehaviorTreeEditorTypes.h
@@ -19,15 +19,12 @@ struct FClassData
 	GENERATED_USTRUCT_BODY()
 
 	FClassData() {}
-	FClassData(UClass* InClass, const FString& InDeprecatedMessage) : 
-		bIsHidden(0), bHideParent(0), Class(InClass), DeprecatedMessage(InDeprecatedMessage)
-	{}
-	FClassData(const FString& InAssetName, const FString& InGeneratedClassPackage, const FString& InClassName, UClass* InClass) :
-		bIsHidden(0), bHideParent(0), Class(InClass), AssetName(InAssetName), GeneratedClassPackage(InGeneratedClassPackage), ClassName(InClassName) 
-	{}
+	FClassData(UClass* InClass, const FString& InDeprecatedMessage);
+	FClassData(const FString& InAssetName, const FString& InGeneratedClassPackage, const FString& InClassName, UClass* InClass);
 
 	FString ToString() const;
 	FString GetClassName() const;
+	FString GetCategory() const;
 	UClass* GetClass(bool bSilent = false);
 	bool IsAbstract() const;
 	
@@ -57,6 +54,10 @@ private:
 	UPROPERTY()
 	FString ClassName;
 
+	/** User-defined category for this class */
+	UPROPERTY()
+	FString Category;
+
 	/** message for deprecated class */
 	FString DeprecatedMessage;
 };
diff --git a/Engine/Source/Editor/BehaviorTreeEditor/Private/BehaviorTreeEditorTypes.cpp b/Engine/Source/Editor/BehaviorTreeEditor/Private/BehaviorTreeEditorTypes.cpp
--- a/Engine/Source/Editor/BehaviorTreeEditor/Private/BehaviorTreeEditorTypes.cpp
+++ b/Engine/Source/Editor/BehaviorTreeEditor/Private/BehaviorTreeEditorTypes.cpp
@@ -16,6 +16,26 @@ const FString UBehaviorTreeEditorTypes::PinCategory_SingleNode("SingleNode");
 
 #define LOCTEXT_NAMESPACE "SClassViewer"
 
+FClassData::FClassData(UClass* InClass, const FString& InDeprecatedMessage) : 
+	bIsHidden(0), 
+	bHideParent(0), 
+	Class(InClass), 
+	DeprecatedMessage(InDeprecatedMessage)
+{
+	Category = GetCategory();
+}
+
+FClassData::FClassData(const FString& InAssetName, const FString& InGeneratedClassPackage, const FString& InClassName, UClass* InClass) :
+	bIsHidden(0), 
+	bHideParent(0), 
+	Class(InClass), 
+	AssetName(InAssetName), 
+	GeneratedClassPackage(InGeneratedClassPackage), 
+	ClassName(InClassName) 
+{
+	Category = GetCategory();
+}
+
 FString FClassData::ToString() const
 {
 	UClass* MyClass = Class.Get();
@@ -45,6 +65,11 @@ FString FClassData::GetClassName() const
 	return Class.IsValid() ? Class->GetName() : ClassName;
 }
 
+FString FClassData::GetCategory() const
+{
+	return Class.IsValid() ? Class->GetMetaData(TEXT("Category")) : Category;
+}
+
 bool FClassData::IsAbstract() const
 {
 	return Class.IsValid() ? Class.Get()->HasAnyClassFlags(CLASS_Abstract) : false;
diff --git a/Engine/Source/Editor/BehaviorTreeEditor/Private/EdGraphSchema_BehaviorTree.cpp b/Engine/Source/Editor/BehaviorTreeEditor/Private/EdGraphSchema_BehaviorTree.cpp
--- a/Engine/Source/Editor/BehaviorTreeEditor/Private/EdGraphSchema_BehaviorTree.cpp
+++ b/Engine/Source/Editor/BehaviorTreeEditor/Private/EdGraphSchema_BehaviorTree.cpp
@@ -224,7 +224,7 @@ UEdGraphSchema_BehaviorTree::UEdGraphSchema_BehaviorTree(const class FPostConstr
 {
 }
 
-TSharedPtr<FBehaviorTreeSchemaAction_NewNode> AddNewNodeAction(FGraphContextMenuBuilder& ContextMenuBuilder, const FString& Category, const FText& MenuDesc, const FString& Tooltip)
+TSharedPtr<FBehaviorTreeSchemaAction_NewNode> AddNewNodeAction(FGraphActionListBuilderBase& ContextMenuBuilder, const FString& Category, const FText& MenuDesc, const FString& Tooltip)
 {
 	TSharedPtr<FBehaviorTreeSchemaAction_NewNode> NewAction = TSharedPtr<FBehaviorTreeSchemaAction_NewNode>(new FBehaviorTreeSchemaAction_NewNode(Category, MenuDesc, Tooltip, 0));
 
@@ -233,7 +233,7 @@ TSharedPtr<FBehaviorTreeSchemaAction_NewNode> AddNewNodeAction(FGraphContextMenu
 	return NewAction;
 }
 
-TSharedPtr<FBehaviorTreeSchemaAction_NewSubNode> AddNewSubNodeAction(FGraphContextMenuBuilder& ContextMenuBuilder, const FString& Category, const FText& MenuDesc, const FString& Tooltip)
+TSharedPtr<FBehaviorTreeSchemaAction_NewSubNode> AddNewSubNodeAction(FGraphActionListBuilderBase& ContextMenuBuilder, const FString& Category, const FText& MenuDesc, const FString& Tooltip)
 {
 	TSharedPtr<FBehaviorTreeSchemaAction_NewSubNode> NewAction = TSharedPtr<FBehaviorTreeSchemaAction_NewSubNode>(new FBehaviorTreeSchemaAction_NewSubNode(Category, MenuDesc, Tooltip, 0));
 	ContextMenuBuilder.AddAction( NewAction );
@@ -274,20 +274,21 @@ void UEdGraphSchema_BehaviorTree::GetGraphNodeContextActions(FGraphContextMenuBu
 		FClassBrowseHelper::GatherClasses(UBTDecorator::StaticClass(), NodeClasses);
 
 		{
+			const FString& Category = UBehaviorTreeGraphNode_CompositeDecorator::StaticClass()->GetMetaData(TEXT("Category"));
 			UBehaviorTreeGraphNode_CompositeDecorator* OpNode = NewObject<UBehaviorTreeGraphNode_CompositeDecorator>(Graph);
-			TSharedPtr<FBehaviorTreeSchemaAction_NewSubNode> AddOpAction = AddNewSubNodeAction(ContextMenuBuilder, TEXT(""), FText::FromString(OpNode->GetNodeTypeDescription()), "");
+			TSharedPtr<FBehaviorTreeSchemaAction_NewSubNode> AddOpAction = AddNewSubNodeAction(ContextMenuBuilder, Category, FText::FromString(OpNode->GetNodeTypeDescription()), "");
 			AddOpAction->ParentNode = Cast<UBehaviorTreeGraphNode>(ContextMenuBuilder.SelectedObjects[0]);
 			AddOpAction->NodeTemplate = OpNode;
 		}
 
-		for (int32 i = 0; i < NodeClasses.Num(); i++)
+		for (const auto& NodeClass : NodeClasses)
 		{
-			const FText NodeTypeName = FText::FromString(FName::NameToDisplayString(NodeClasses[i].ToString(), false));
+			const FText NodeTypeName = FText::FromString(FName::NameToDisplayString(NodeClass.ToString(), false));
 
 			UBehaviorTreeGraphNode* OpNode = NewObject<UBehaviorTreeGraphNode_Decorator>(Graph);
-			OpNode->ClassData = NodeClasses[i];
+			OpNode->ClassData = NodeClass;
 
-			TSharedPtr<FBehaviorTreeSchemaAction_NewSubNode> AddOpAction = AddNewSubNodeAction(ContextMenuBuilder, TEXT(""), NodeTypeName, "");
+			TSharedPtr<FBehaviorTreeSchemaAction_NewSubNode> AddOpAction = AddNewSubNodeAction(ContextMenuBuilder, NodeClass.GetCategory(), NodeTypeName, "");
 			AddOpAction->ParentNode = Cast<UBehaviorTreeGraphNode>(ContextMenuBuilder.SelectedObjects[0]);
 			AddOpAction->NodeTemplate = OpNode;
 		}
@@ -295,14 +296,14 @@ void UEdGraphSchema_BehaviorTree::GetGraphNodeContextActions(FGraphContextMenuBu
 	else if (SubNodeType == ESubNode::Service)
 	{
 		FClassBrowseHelper::GatherClasses(UBTService::StaticClass(), NodeClasses);
-		for (int32 i = 0; i < NodeClasses.Num(); i++)
+		for (const auto& NodeClass : NodeClasses)
 		{
-			const FText NodeTypeName = FText::FromString(FName::NameToDisplayString(NodeClasses[i].ToString(), false));
+			const FText NodeTypeName = FText::FromString(FName::NameToDisplayString(NodeClass.ToString(), false));
 
 			UBehaviorTreeGraphNode* OpNode = NewObject<UBehaviorTreeGraphNode_Service>(Graph);
-			OpNode->ClassData = NodeClasses[i];
+			OpNode->ClassData = NodeClass;
 
-			TSharedPtr<FBehaviorTreeSchemaAction_NewSubNode> AddOpAction = AddNewSubNodeAction(ContextMenuBuilder, TEXT(""), NodeTypeName, "");
+			TSharedPtr<FBehaviorTreeSchemaAction_NewSubNode> AddOpAction = AddNewSubNodeAction(ContextMenuBuilder, NodeClass.GetCategory(), NodeTypeName, "");
 			AddOpAction->ParentNode = Cast<UBehaviorTreeGraphNode>(ContextMenuBuilder.SelectedObjects[0]);
 			AddOpAction->NodeTemplate = OpNode;
 		}
@@ -323,47 +324,55 @@ void UEdGraphSchema_BehaviorTree::GetGraphContextActions(FGraphContextMenuBuilde
 
 	if (bAllowComposites)
 	{
+		FCategorizedGraphActionListBuilder CompositesBuilder(TEXT("Composites"));
+
 		TArray<FClassData> NodeClasses;
 		FClassBrowseHelper::GatherClasses(UBTCompositeNode::StaticClass(), NodeClasses);
 
 		const FString ParallelClassName = UBTComposite_SimpleParallel::StaticClass()->GetName();
 
-		for (int32 i = 0; i < NodeClasses.Num(); i++)
+		for (const auto& NodeClass : NodeClasses)
 		{
-			const FText NodeTypeName = FText::FromString(FName::NameToDisplayString(NodeClasses[i].ToString(), false));
+			const FText NodeTypeName = FText::FromString(FName::NameToDisplayString(NodeClass.ToString(), false));
 
-			TSharedPtr<FBehaviorTreeSchemaAction_NewNode> AddOpAction = AddNewNodeAction(ContextMenuBuilder, TEXT("Composites"), NodeTypeName, "");
+			TSharedPtr<FBehaviorTreeSchemaAction_NewNode> AddOpAction = AddNewNodeAction(CompositesBuilder, NodeClass.GetCategory(), NodeTypeName, "");
 
-			UBehaviorTreeGraphNode* OpNode = (NodeClasses[i].GetClassName() == ParallelClassName) ?
+			UBehaviorTreeGraphNode* OpNode = (NodeClass.GetClassName() == ParallelClassName) ?
 				NewObject<UBehaviorTreeGraphNode_SimpleParallel>(ContextMenuBuilder.OwnerOfTemporaries) :				
 				NewObject<UBehaviorTreeGraphNode_Composite>(ContextMenuBuilder.OwnerOfTemporaries);
 
-			OpNode->ClassData = NodeClasses[i];
+			OpNode->ClassData = NodeClass;
 			AddOpAction->NodeTemplate = OpNode;
 		}
+
+		ContextMenuBuilder.Append(CompositesBuilder);
 	}
 
 	if (bAllowTasks)
 	{
+		FCategorizedGraphActionListBuilder TasksBuilder(TEXT("Tasks"));
+
 		TArray<FClassData> NodeClasses;
 		FClassBrowseHelper::GatherClasses(UBTTaskNode::StaticClass(), NodeClasses);
 
-		for (int32 i = 0; i < NodeClasses.Num(); i++)
+		for (const auto& NodeClass : NodeClasses)
 		{
-			const FText NodeTypeName = FText::FromString(FName::NameToDisplayString(NodeClasses[i].ToString(), false));
+			const FText NodeTypeName = FText::FromString(FName::NameToDisplayString(NodeClass.ToString(), false));
 
-			TSharedPtr<FBehaviorTreeSchemaAction_NewNode> AddOpAction = AddNewNodeAction(ContextMenuBuilder, TEXT("Tasks"), NodeTypeName, "");
+			TSharedPtr<FBehaviorTreeSchemaAction_NewNode> AddOpAction = AddNewNodeAction(TasksBuilder, NodeClass.GetCategory(), NodeTypeName, "");
 			UClass* GraphNodeClass = UBehaviorTreeGraphNode_Task::StaticClass();
 			
-			if (NodeClasses[i].GetClassName() == UBTTask_RunBehavior::StaticClass()->GetName())
+			if (NodeClass.GetClassName() == UBTTask_RunBehavior::StaticClass()->GetName())
 			{
 				GraphNodeClass = UBehaviorTreeGraphNode_SubtreeTask::StaticClass();
 			}
 
 			UBehaviorTreeGraphNode* OpNode = ConstructObject<UBehaviorTreeGraphNode>(GraphNodeClass, ContextMenuBuilder.OwnerOfTemporaries);
-			OpNode->ClassData = NodeClasses[i];
+			OpNode->ClassData = NodeClass;
 			AddOpAction->NodeTemplate = OpNode;
 		}
+
+		ContextMenuBuilder.Append(TasksBuilder);
 	}
 	
 	if (bNoParent)
diff --git a/Engine/Source/Editor/EditorStyle/Private/SlateEditorStyle.cpp b/Engine/Source/Editor/EditorStyle/Private/SlateEditorStyle.cpp
--- a/Engine/Source/Editor/EditorStyle/Private/SlateEditorStyle.cpp
+++ b/Engine/Source/Editor/EditorStyle/Private/SlateEditorStyle.cpp
@@ -2802,12 +2802,14 @@ void FSlateEditorStyle::FStyle::SetupGraphEditorStyles()
 			Set( "BTEditor.Graph.BTNode.Composite.SimpleParallel.Icon", new IMAGE_BRUSH( "/BehaviorTree/Icons/Simple_Parallel_24x", Icon24x24 ) );
 			Set( "BTEditor.Graph.BTNode.Decorator.Blackboard.Icon", new IMAGE_BRUSH( "/BehaviorTree/Icons/Blackboard_24x", Icon24x24 ) );
 			Set( "BTEditor.Graph.BTNode.Decorator.CompareBlackboardEntries.Icon", new IMAGE_BRUSH( "/BehaviorTree/Icons/Compare_Blackboard_Entries_24x", Icon24x24 ) );
+			Set( "BTEditor.Graph.BTNode.Decorator.Conditional.Icon", new IMAGE_BRUSH( "/BehaviorTree/Icons/Conditional_Decorator_24x", Icon24x24 ) );
 			Set( "BTEditor.Graph.BTNode.Decorator.ConeCheck.Icon", new IMAGE_BRUSH( "/BehaviorTree/Icons/Cone_Check_24x", Icon24x24 ) );
 			Set( "BTEditor.Graph.BTNode.Decorator.Cooldown.Icon", new IMAGE_BRUSH( "/BehaviorTree/Icons/Cooldown_24x", Icon24x24 ) );
 			Set( "BTEditor.Graph.BTNode.Decorator.DoesPathExist.Icon", new IMAGE_BRUSH( "/BehaviorTree/Icons/Does_Path_Exist_24x", Icon24x24 ) );
 			Set( "BTEditor.Graph.BTNode.Decorator.ForceSuccess.Icon", new IMAGE_BRUSH( "/BehaviorTree/Icons/Force_Success_24x", Icon24x24 ) );
 			Set( "BTEditor.Graph.BTNode.Decorator.KeepInCone.Icon", new IMAGE_BRUSH( "/BehaviorTree/Icons/Keep_In_Cone_24x", Icon24x24 ) );
 			Set( "BTEditor.Graph.BTNode.Decorator.Loop.Icon", new IMAGE_BRUSH( "/BehaviorTree/Icons/Loop_24x", Icon24x24 ) );
+			Set( "BTEditor.Graph.BTNode.Decorator.NonConditional.Icon", new IMAGE_BRUSH( "/BehaviorTree/Icons/Non_Conditional_Decorator_24x", Icon24x24 ) );
 			Set( "BTEditor.Graph.BTNode.Decorator.Optional.Icon", new IMAGE_BRUSH( "/BehaviorTree/Icons/Optional_24x", Icon24x24 ) );
 			Set( "BTEditor.Graph.BTNode.Decorator.ReachedMoveGoal.Icon", new IMAGE_BRUSH( "/BehaviorTree/Icons/Reached_Move_Goal_24x", Icon24x24 ) );
 			Set( "BTEditor.Graph.BTNode.Decorator.TimeLimit.Icon", new IMAGE_BRUSH( "/BehaviorTree/Icons/Time_Limit_24x", Icon24x24 ) );
diff --git a/Engine/Source/Editor/KismetCompiler/Private/KismetCompiler.cpp b/Engine/Source/Editor/KismetCompiler/Private/KismetCompiler.cpp
--- a/Engine/Source/Editor/KismetCompiler/Private/KismetCompiler.cpp
+++ b/Engine/Source/Editor/KismetCompiler/Private/KismetCompiler.cpp
@@ -1625,6 +1625,12 @@ void FKismetCompilerContext::FinishCompilingClass(UClass* Class)
 		{
 			Class->SetMetaData(TEXT("AutoCollapseCategories"), *ParentClass->GetMetaData("AutoCollapseCategories"));
 		}
+
+		// Add a category if one has been specified
+		if(Blueprint->BlueprintCategory.Len() > 0)
+		{
+			Class->SetMetaData(TEXT("Category"), *Blueprint->BlueprintCategory);
+		}
 #endif
 
 		// Add in additional flags implied by the blueprint
diff --git a/Engine/Source/Runtime/Engine/Classes/Engine/Blueprint.h b/Engine/Source/Runtime/Engine/Classes/Engine/Blueprint.h
--- a/Engine/Source/Runtime/Engine/Classes/Engine/Blueprint.h
+++ b/Engine/Source/Runtime/Engine/Classes/Engine/Blueprint.h
@@ -295,6 +295,10 @@ class ENGINE_API UBlueprint : public UBlueprintCore
 	UPROPERTY(EditAnywhere, Category=BlueprintOption)
 	FString BlueprintDescription;
 
+	/** The category of the Blueprint, used to organize this Blueprint class when displayed in palette windows */
+	UPROPERTY(EditAnywhere, Category=BlueprintOption)
+	FString BlueprintCategory;
+
 	/** TRUE to show a warning when attempting to start in PIE and there is a compiler error on this Blueprint */
 	UPROPERTY(transient)
 	bool bDisplayCompilePIEWarning;
