Dont block graph connections in behavior trees, replace existing instead

Reworked connection logic in the behavior tree graph schema.
Found that while the code to check while dragging/hovering was switched when dropping interms of pin direction. This made denying connections in directed graphs (like BTs) problematic. Fixed this by just reversing the pin order on drop, tested in Blueprint & it seems to be working fine.
Also fixed creating graph cycles, which caused a crash.

TTP# 339732 - Editor: Behaviour Tree: Dragging link should replace existing, not be blocked

diff --git a/Engine/Source/Editor/BehaviorTreeEditor/Private/EdGraphSchema_BehaviorTree.cpp b/Engine/Source/Editor/BehaviorTreeEditor/Private/EdGraphSchema_BehaviorTree.cpp
--- a/Engine/Source/Editor/BehaviorTreeEditor/Private/EdGraphSchema_BehaviorTree.cpp
+++ b/Engine/Source/Editor/BehaviorTreeEditor/Private/EdGraphSchema_BehaviorTree.cpp
@@ -517,37 +517,85 @@ const FPinConnectionResponse UEdGraphSchema_BehaviorTree::CanCreateConnection(co
 		return FPinConnectionResponse(CONNECT_RESPONSE_DISALLOW, LOCTEXT("PinErrorOnlyTask","Only task nodes are allowed"));
 	}
 
-	if (((bPinAIsSingleNode || bPinAIsSingleTask || bPinAIsSingleComposite) && PinA->LinkedTo.Num() > 0) ||
-		((bPinBIsSingleNode || bPinBIsSingleTask || bPinBIsSingleComposite) && PinB->LinkedTo.Num() > 0))
+	// Compare the directions
+	if ((PinA->Direction == EGPD_Input) && (PinB->Direction == EGPD_Input))
 	{
-		return FPinConnectionResponse(CONNECT_RESPONSE_DISALLOW, LOCTEXT("PinErrorSingleNode","Can't connect multiple nodes"));
+		return FPinConnectionResponse(CONNECT_RESPONSE_DISALLOW, LOCTEXT("PinErrorInput", "Can't connect input node to input node"));
 	}
-
-	// Compare the directions
-	bool bDirectionsOK = false;
-
-	if ((PinA->Direction == EGPD_Input) && (PinB->Direction == EGPD_Output))
+	else if ((PinB->Direction == EGPD_Output) && (PinA->Direction == EGPD_Output))
 	{
-		bDirectionsOK = true;
+		return FPinConnectionResponse(CONNECT_RESPONSE_DISALLOW, LOCTEXT("PinErrorOutput", "Can't connect output node to output node"));
 	}
-	else if ((PinB->Direction == EGPD_Input) && (PinA->Direction == EGPD_Output))
+	else if ((PinA->Direction == EGPD_Input) && (PinB->Direction == EGPD_Output))
 	{
-		bDirectionsOK = true;
+		return FPinConnectionResponse(CONNECT_RESPONSE_DISALLOW, LOCTEXT("PinErrorWrongDirection", "Can't connect input node to output node"));
 	}
 
-	if (bDirectionsOK)
+	class FNodeVisitorCycleChecker
 	{
-		if ( (PinA->Direction == EGPD_Input && PinA->LinkedTo.Num()>0) || (PinB->Direction == EGPD_Input && PinB->LinkedTo.Num()>0))
+	public:
+		/** Check whether a loop in the graph would be caused by linking the passed-in nodes */
+		bool CheckForLoop(UEdGraphNode* StartNode, UEdGraphNode* EndNode)
 		{
-			return FPinConnectionResponse(CONNECT_RESPONSE_DISALLOW, LOCTEXT("PinErrorAlreadyConnected","Already connected with other"));
+			VisitedNodes.Add(EndNode);
+			return TraverseInputNodesToRoot(StartNode);
 		}
+
+	private:
+		/** 
+		 * Helper function for CheckForLoop()
+		 * @param	Node	The node to start traversal at
+		 * @return true if we reached a root node (i.e. a node with no input pins), false if we encounter a node we have already seen
+		 */
+		bool TraverseInputNodesToRoot(UEdGraphNode* Node)
+		{
+			VisitedNodes.Add(Node);
+
+			// Follow every input pin until we cant any more ('root') or we reach a node we have seen (cycle)
+			for (int32 PinIndex = 0; PinIndex < Node->Pins.Num(); ++PinIndex)
+			{
+				UEdGraphPin* MyPin = Node->Pins[PinIndex];
+
+				if (MyPin->Direction == EGPD_Input)
+				{
+					for (int32 LinkedPinIndex = 0; LinkedPinIndex < MyPin->LinkedTo.Num(); ++LinkedPinIndex)
+					{
+						UEdGraphPin* OtherPin = MyPin->LinkedTo[LinkedPinIndex];
+						if( OtherPin )
+						{
+							UEdGraphNode* OtherNode = OtherPin->GetOwningNode();
+							if (VisitedNodes.Contains(OtherNode))
+							{
+								return false;
+							}
+							else
+							{
+								return TraverseInputNodesToRoot(OtherNode);
+							}
+						}
+					}
+				}
+			}
+
+			return true;
+		}
+
+		TSet<UEdGraphNode*> VisitedNodes;
+	};
+
+	// check for cycles
+	FNodeVisitorCycleChecker CycleChecker;
+	if(!CycleChecker.CheckForLoop(PinA->GetOwningNode(), PinB->GetOwningNode()))
+	{
+		return FPinConnectionResponse(CONNECT_RESPONSE_DISALLOW, LOCTEXT("PinErrorcycle", "Can't create a graph cycle"));
 	}
-	else
+
+	if(PinB->LinkedTo.Num() > 0)
 	{
-		return FPinConnectionResponse(CONNECT_RESPONSE_DISALLOW, TEXT(""));
+		return FPinConnectionResponse(CONNECT_RESPONSE_BREAK_OTHERS_B, LOCTEXT("PinConnectReplace", "Replace connection"));
 	}
 
-	return FPinConnectionResponse(CONNECT_RESPONSE_MAKE, TEXT(""));
+	return FPinConnectionResponse(CONNECT_RESPONSE_MAKE, LOCTEXT("PinConnect", "Connect nodes"));
 }
 
 const FPinConnectionResponse UEdGraphSchema_BehaviorTree::CanMergeNodes(const UEdGraphNode* NodeA, const UEdGraphNode* NodeB) const
diff --git a/Engine/Source/Editor/GraphEditor/Private/DragConnection.cpp b/Engine/Source/Editor/GraphEditor/Private/DragConnection.cpp
--- a/Engine/Source/Editor/GraphEditor/Private/DragConnection.cpp
+++ b/Engine/Source/Editor/GraphEditor/Private/DragConnection.cpp
@@ -157,11 +157,11 @@ FReply FDragConnection::DroppedOnPin(FVector2D ScreenPosition, FVector2D GraphPo
 
 	const FScopedTransaction Transaction( NSLOCTEXT("UnrealEd", "GraphEd_CreateConnection", "Create Pin Link") );
 
-	UEdGraphPin* PinA = GetHoveredPin();
+	UEdGraphPin* PinB = GetHoveredPin();
 	bool bError = false;
 	TSet<UEdGraphNode*> NodeList;
 
-	for (UEdGraphPin* PinB : ValidSourcePins)
+	for (UEdGraphPin* PinA : ValidSourcePins)
 	{
 		if ((PinA != NULL) && (PinB != NULL))
 		{
