Managing and creating Blackboard asset improvements

Added factory for Blackboard assets so they can have their own entry & icon in the context menu.
Added toolbar shortcut to create a new Blackboard asset alongside the BT.
Fixed selecting new Blackboard sometimes not displaying in Blackboard mode.
Displayed root node properties when nothing is selected so blackboard assets can be chosen.
Also added missing icons for small toolbars

TTP# 339961 - [Editor][AI] There's no way to create a blackboard asset
TTP# 339731 - Editor: Behaviour Tree: Root node UI improvements

diff --git a/Engine/Source/Editor/BehaviorTreeEditor/BehaviorTreeEditor.Build.cs b/Engine/Source/Editor/BehaviorTreeEditor/BehaviorTreeEditor.Build.cs
--- a/Engine/Source/Editor/BehaviorTreeEditor/BehaviorTreeEditor.Build.cs
+++ b/Engine/Source/Editor/BehaviorTreeEditor/BehaviorTreeEditor.Build.cs
@@ -18,7 +18,8 @@ public class BehaviorTreeEditor : ModuleRules
             new string[] {
 				"AssetRegistry",
 				"AssetTools",
-                "PropertyEditor"
+                "PropertyEditor",
+				"ContentBrowser"
 			}
 		);
 
@@ -52,7 +53,8 @@ public class BehaviorTreeEditor : ModuleRules
                 "WorkspaceMenuStructure",
                 "PropertyEditor",
 				"AssetTools",
-				"AssetRegistry"
+				"AssetRegistry",
+				"ContentBrowser"
             }
 		);
 	}
diff --git a/Engine/Source/Editor/BehaviorTreeEditor/Classes/BlackboardDataFactory.h b/Engine/Source/Editor/BehaviorTreeEditor/Classes/BlackboardDataFactory.h
new file mode 100644
--- /dev/null
+++ b/Engine/Source/Editor/BehaviorTreeEditor/Classes/BlackboardDataFactory.h
@@ -0,0 +1,15 @@
+// Copyright 1998-2014 Epic Games, Inc. All Rights Reserved.
+
+#pragma once
+#include "BlackboardDataFactory.generated.h"
+
+UCLASS()
+class UBlackboardDataFactory : public UFactory
+{
+	GENERATED_UCLASS_BODY()
+
+	// UFactory interface
+	virtual UObject* FactoryCreateNew(UClass* Class,UObject* InParent,FName Name,EObjectFlags Flags,UObject* Context,FFeedbackContext* Warn) override;
+	virtual bool CanCreateNew() const override;
+	// End of UFactory interface
+};
diff --git a/Engine/Source/Editor/BehaviorTreeEditor/Private/BehaviorTreeEditor.cpp b/Engine/Source/Editor/BehaviorTreeEditor/Private/BehaviorTreeEditor.cpp
--- a/Engine/Source/Editor/BehaviorTreeEditor/Private/BehaviorTreeEditor.cpp
+++ b/Engine/Source/Editor/BehaviorTreeEditor/Private/BehaviorTreeEditor.cpp
@@ -32,6 +32,7 @@
 #include "ClassViewerModule.h"
 #include "AssetRegistryModule.h"
 #include "AssetToolsModule.h"
+#include "ContentBrowserModule.h"
 
 #define LOCTEXT_NAMESPACE "BehaviorTreeEditor"
 
@@ -125,6 +126,11 @@ void FBehaviorTreeEditor::NotifyPostChange( const FPropertyChangedEvent& Propert
 {
 	if(PropertyChangedEvent.ChangeType != EPropertyChangeType::Interactive)
 	{
+		if(PropertyChangedEvent.Property != nullptr && PropertyChangedEvent.Property->GetFName() == TEXT("BlackboardAsset"))
+		{
+			BlackboardData = BehaviorTree->BlackboardAsset;
+		}
+
 		if(BlackboardView.IsValid())
 		{
 			BlackboardView->SetObject(GetBlackboardData());
@@ -783,7 +789,25 @@ void FBehaviorTreeEditor::OnSelectedNodesChanged(const TSet<class UObject*>& New
 	}
 	else if (DetailsView.IsValid())
 	{
-		DetailsView->SetObject(NULL);
+		if (Selection.Num() == 0)
+		{
+			// if nothing is selected, display the root
+			UBehaviorTreeGraphNode* RootNode = nullptr;
+			for (const auto& Node : MyGraph->Nodes)
+			{
+				RootNode = Cast<UBehaviorTreeGraphNode_Root>(Node);
+				if (RootNode != nullptr)
+				{
+					break;
+				}
+			}
+
+			DetailsView->SetObject(RootNode);
+		}
+		else
+		{
+			DetailsView->SetObject(nullptr);
+		}
 	}
 
 	MyGraph->UpdateAbortHighlight(Mode0, Mode1);
@@ -907,7 +931,8 @@ void FBehaviorTreeEditor::GetAbortModePreview(const class UBTDecorator* Decorato
 void FBehaviorTreeEditor::CreateInternalWidgets()
 {
 	FPropertyEditorModule& PropertyEditorModule = FModuleManager::GetModuleChecked<FPropertyEditorModule>( "PropertyEditor" );
-	const FDetailsViewArgs DetailsViewArgs( false, false, true, true, false );
+	FDetailsViewArgs DetailsViewArgs( false, false, true, true, false );
+	DetailsViewArgs.NotifyHook = this;
 	DetailsView = PropertyEditorModule.CreateDetailView( DetailsViewArgs );
 	DetailsView->SetObject( NULL );
 	DetailsView->SetIsPropertyVisibleDelegate(FIsPropertyVisible::CreateSP(this, &FBehaviorTreeEditor::IsPropertyVisible));
@@ -947,6 +972,11 @@ void FBehaviorTreeEditor::BindCommonCommands()
 			FExecuteAction::CreateSP(this, &FBehaviorTreeEditor::SearchTree),
 			FCanExecuteAction::CreateSP(this, &FBehaviorTreeEditor::CanSearchTree)
 			);
+
+	ToolkitCommands->MapAction(FBTCommonCommands::Get().NewBlackboard,
+			FExecuteAction::CreateSP(this, &FBehaviorTreeEditor::CreateNewBlackboard),
+			FCanExecuteAction::CreateSP(this, &FBehaviorTreeEditor::CanCreateNewBlackboard)
+			);
 }
 
 void FBehaviorTreeEditor::SearchTree()
@@ -2119,4 +2149,25 @@ bool FBehaviorTreeEditor::IsNewServiceComboVisible() const
 	return FClassBrowseHelper::IsMoreThanOneServiceClassAvailable();
 }
 
+void FBehaviorTreeEditor::CreateNewBlackboard()
+{
+	FString PathName = BehaviorTree->GetOutermost()->GetPathName();
+	PathName = FPaths::GetPath(PathName);
+	const FString PathNameWithFilename = PathName / LOCTEXT("NewBlackboardName", "NewBlackboardData").ToString();
+
+	FString Name;
+	FString PackageName;
+	FAssetToolsModule& AssetToolsModule = FModuleManager::GetModuleChecked<FAssetToolsModule>("AssetTools");
+	AssetToolsModule.Get().CreateUniqueAssetName(PathNameWithFilename, TEXT(""), PackageName, Name);
+
+	UDataAssetFactory* DataAssetFactory = ConstructObject<UDataAssetFactory>(UDataAssetFactory::StaticClass());
+	FContentBrowserModule& ContentBrowserModule = FModuleManager::LoadModuleChecked<FContentBrowserModule>("ContentBrowser");
+	ContentBrowserModule.Get().CreateNewAsset(Name, PathName, UBlackboardData::StaticClass(), DataAssetFactory);
+}
+
+bool FBehaviorTreeEditor::CanCreateNewBlackboard() const
+{
+	return !IsDebuggerReady();
+}
+
 #undef LOCTEXT_NAMESPACE
diff --git a/Engine/Source/Editor/BehaviorTreeEditor/Private/BehaviorTreeEditor.h b/Engine/Source/Editor/BehaviorTreeEditor/Private/BehaviorTreeEditor.h
--- a/Engine/Source/Editor/BehaviorTreeEditor/Private/BehaviorTreeEditor.h
+++ b/Engine/Source/Editor/BehaviorTreeEditor/Private/BehaviorTreeEditor.h
@@ -232,6 +232,12 @@ public:
 	/** Whether the combo button to create a new Blueprint-based service from all available base classes is visible */
 	bool IsNewServiceComboVisible() const;
 
+	/** Create a new Blackboard alongside the currently-edited behavior tree */
+	void CreateNewBlackboard();
+
+	/** Whether we can currently create a new Blackboard */
+	bool CanCreateNewBlackboard() const;
+
 protected:
 	/** Called when "Save" is clicked for this asset */
 	virtual void SaveAsset_Execute() override;
diff --git a/Engine/Source/Editor/BehaviorTreeEditor/Private/BehaviorTreeEditorCommands.cpp b/Engine/Source/Editor/BehaviorTreeEditor/Private/BehaviorTreeEditorCommands.cpp
--- a/Engine/Source/Editor/BehaviorTreeEditor/Private/BehaviorTreeEditorCommands.cpp
+++ b/Engine/Source/Editor/BehaviorTreeEditor/Private/BehaviorTreeEditorCommands.cpp
@@ -13,6 +13,7 @@ FBTCommonCommands::FBTCommonCommands()
 void FBTCommonCommands::RegisterCommands()
 {
 	UI_COMMAND(SearchBT, "Search", "Search this Behavior Tree.", EUserInterfaceActionType::Button, FInputGesture(EModifierKey::Control, EKeys::F));
+	UI_COMMAND(NewBlackboard, "New Blackboard", "Create a new Blackboard Data Asset", EUserInterfaceActionType::Button, FInputGesture());
 }
 
 FBTDebuggerCommands::FBTDebuggerCommands() 
diff --git a/Engine/Source/Editor/BehaviorTreeEditor/Private/BehaviorTreeEditorCommands.h b/Engine/Source/Editor/BehaviorTreeEditor/Private/BehaviorTreeEditorCommands.h
--- a/Engine/Source/Editor/BehaviorTreeEditor/Private/BehaviorTreeEditorCommands.h
+++ b/Engine/Source/Editor/BehaviorTreeEditor/Private/BehaviorTreeEditorCommands.h
@@ -10,6 +10,7 @@ public:
 	FBTCommonCommands();
 
 	TSharedPtr<FUICommandInfo> SearchBT;
+	TSharedPtr<FUICommandInfo> NewBlackboard;
 
 	/** Initialize commands */
 	virtual void RegisterCommands() override;
diff --git a/Engine/Source/Editor/BehaviorTreeEditor/Private/BehaviorTreeEditorToolbar.cpp b/Engine/Source/Editor/BehaviorTreeEditor/Private/BehaviorTreeEditorToolbar.cpp
--- a/Engine/Source/Editor/BehaviorTreeEditor/Private/BehaviorTreeEditorToolbar.cpp
+++ b/Engine/Source/Editor/BehaviorTreeEditor/Private/BehaviorTreeEditorToolbar.cpp
@@ -150,6 +150,12 @@ void FBehaviorTreeEditorToolbar::FillBehaviorTreeToolbar(FToolBarBuilder& Toolba
 
 	if (!BehaviorTreeEditorPtr->IsDebuggerReady() && BehaviorTreeEditorPtr->GetCurrentMode() == FBehaviorTreeEditor::BehaviorTreeMode)
 	{
+		ToolbarBuilder.BeginSection("Blackboard");
+		{
+			ToolbarBuilder.AddToolBarButton(FBTCommonCommands::Get().NewBlackboard);
+		}
+		ToolbarBuilder.EndSection();
+
 		ToolbarBuilder.BeginSection("BehaviorTree");
 		{
 			const FText NewTaskLabel = LOCTEXT( "NewTask_Label", "New Task" );
diff --git a/Engine/Source/Editor/BehaviorTreeEditor/Private/BlackboardDataFactory.cpp b/Engine/Source/Editor/BehaviorTreeEditor/Private/BlackboardDataFactory.cpp
new file mode 100644
--- /dev/null
+++ b/Engine/Source/Editor/BehaviorTreeEditor/Private/BlackboardDataFactory.cpp
@@ -0,0 +1,30 @@
+// Copyright 1998-2014 Epic Games, Inc. All Rights Reserved.
+
+#include "BehaviorTreeEditorPrivatePCH.h"
+#include "BehaviorTree/BlackboardData.h"
+#include "BlackboardDataFactory.h"
+
+#define LOCTEXT_NAMESPACE "BehaviorTreeFactory"
+
+UBlackboardDataFactory::UBlackboardDataFactory(const class FPostConstructInitializeProperties& PCIP)
+	: Super(PCIP)
+{
+	SupportedClass = UBlackboardData::StaticClass();
+	bCreateNew = true;
+	bEditAfterNew = true;
+}
+
+bool UBlackboardDataFactory::CanCreateNew() const
+{
+	bool bBehaviorTreeNewAssetsEnabled = false;
+	GConfig->GetBool(TEXT("BehaviorTreesEd"), TEXT("BehaviorTreeNewAssetsEnabled"), bBehaviorTreeNewAssetsEnabled, GEngineIni);
+	return (bBehaviorTreeNewAssetsEnabled || GetDefault<UEditorExperimentalSettings>()->bBehaviorTreeEditor);
+}
+
+UObject* UBlackboardDataFactory::FactoryCreateNew(UClass* Class,UObject* InParent,FName Name,EObjectFlags Flags,UObject* Context,FFeedbackContext* Warn)
+{
+	check(Class->IsChildOf(UBlackboardData::StaticClass()));
+	return ConstructObject<UBlackboardData>(Class, InParent, Name, Flags);
+}
+
+#undef LOCTEXT_NAMESPACE
\ No newline at end of file
diff --git a/Engine/Source/Editor/EditorStyle/Private/SlateEditorStyle.cpp b/Engine/Source/Editor/EditorStyle/Private/SlateEditorStyle.cpp
--- a/Engine/Source/Editor/EditorStyle/Private/SlateEditorStyle.cpp
+++ b/Engine/Source/Editor/EditorStyle/Private/SlateEditorStyle.cpp
@@ -2936,10 +2936,10 @@ void FSlateEditorStyle::FStyle::SetupGraphEditorStyles()
 			Set( "BTEditor.Blackboard.NewEntry", new IMAGE_BRUSH( "BehaviorTree/Blackboard_AddKey_40x", Icon40x40 ) );
 			Set( "BTEditor.Blackboard.NewEntry.Small", new IMAGE_BRUSH( "BehaviorTree/Blackboard_AddKey_40x", Icon20x20 ) );
 
-			Set( "BTEditor.SwitchToBehaviorTreeMode", new IMAGE_BRUSH( "BehaviorTree/BehaviorTreeMode_40x", FVector2D(40,40)));
-			Set( "BTEditor.SwitchToBehaviorTreeMode.Small", new IMAGE_BRUSH( "BehaviorTree/BehaviorTreeMode_16x", FVector2D(16,16)));
-			Set( "BTEditor.SwitchToBlackboardMode", new IMAGE_BRUSH( "BehaviorTree/BlackboardMode_40x", FVector2D(40,40)));
-			Set( "BTEditor.SwitchToBlackboardMode.Small", new IMAGE_BRUSH( "BehaviorTree/BlackboardMode_16x", FVector2D(16,16)));
+			Set( "BTEditor.SwitchToBehaviorTreeMode", new IMAGE_BRUSH( "BehaviorTree/BehaviorTreeMode_40x", Icon40x40));
+			Set( "BTEditor.SwitchToBehaviorTreeMode.Small", new IMAGE_BRUSH( "BehaviorTree/BehaviorTreeMode_20x", Icon20x20));
+			Set( "BTEditor.SwitchToBlackboardMode", new IMAGE_BRUSH( "BehaviorTree/BlackboardMode_40x", Icon40x40));
+			Set( "BTEditor.SwitchToBlackboardMode.Small", new IMAGE_BRUSH( "BehaviorTree/BlackboardMode_20x", Icon20x20));
 
 			// Blackboard classes
 			Set( "ClassIcon.BlackboardKeyType_Bool", new IMAGE_BRUSH( "Icons/pill_16x", Icon16x16, FLinearColor(0.300000f, 0.0f, 0.0f, 1.0f) ) );
@@ -2954,12 +2954,14 @@ void FSlateEditorStyle::FStyle::SetupGraphEditorStyles()
 			Set( "ClassIcon.BlackboardKeyType_String", new IMAGE_BRUSH( "Icons/pill_16x", Icon16x16, FLinearColor(1.0f, 0.0f, 0.660537f, 1.0f) ) );
 			Set( "ClassIcon.BlackboardKeyType_Vector", new IMAGE_BRUSH( "Icons/pill_16x", Icon16x16, FLinearColor(1.0f, 0.591255f, 0.016512f, 1.0f) ) );
 
-			Set( "BTEditor.Graph.NewTask", new IMAGE_BRUSH( "BehaviorTree/NewTask_40x", FVector2D(40,40)));
-			Set( "BTEditor.Graph.NewTask.Small", new IMAGE_BRUSH( "BehaviorTree/NewTask_40x", FVector2D(16,16)));
-			Set( "BTEditor.Graph.NewDecorator", new IMAGE_BRUSH( "BehaviorTree/NewDecorator_40x", FVector2D(40,40)));
-			Set( "BTEditor.Graph.NewDecorator.Small", new IMAGE_BRUSH( "BehaviorTree/NewDecorator_40x", FVector2D(16,16)));
-			Set( "BTEditor.Graph.NewService", new IMAGE_BRUSH( "BehaviorTree/NewService_40x", FVector2D(40,40)));
-			Set( "BTEditor.Graph.NewService.Small", new IMAGE_BRUSH( "BehaviorTree/NewService_40x", FVector2D(16,16)));
+			Set( "BTEditor.Common.NewBlackboard", new IMAGE_BRUSH( "BehaviorTree/NewBlackboard_40x", Icon40x40));
+			Set( "BTEditor.Common.NewBlackboard.Small", new IMAGE_BRUSH( "BehaviorTree/NewBlackboard_20x", Icon20x20));
+			Set( "BTEditor.Graph.NewTask", new IMAGE_BRUSH( "BehaviorTree/NewTask_40x", Icon40x40));
+			Set( "BTEditor.Graph.NewTask.Small", new IMAGE_BRUSH( "BehaviorTree/NewTask_20x", Icon20x20));
+			Set( "BTEditor.Graph.NewDecorator", new IMAGE_BRUSH( "BehaviorTree/NewDecorator_40x", Icon40x40));
+			Set( "BTEditor.Graph.NewDecorator.Small", new IMAGE_BRUSH( "BehaviorTree/NewDecorator_20x", Icon20x20));
+			Set( "BTEditor.Graph.NewService", new IMAGE_BRUSH( "BehaviorTree/NewService_40x", Icon40x40));
+			Set( "BTEditor.Graph.NewService.Small", new IMAGE_BRUSH( "BehaviorTree/NewService_20x", Icon20x20));
 		}
 
 		// Visible on hover button for transition node
@@ -4249,6 +4251,7 @@ void FSlateEditorStyle::FStyle::SetupClassIconsAndThumbnails()
 			TEXT("AnimSequence"),
 			TEXT("AtmosphericFog"),
 			TEXT("BehaviorTree"),
+			TEXT("BlackboardData"),
 			TEXT("BlendSpace"),
 			TEXT("BlendSpace1D"),
 			TEXT("BlockingVolume"),
diff --git a/Engine/Source/Editor/UnrealEd/Classes/Factories/DataAssetFactory.h b/Engine/Source/Editor/UnrealEd/Classes/Factories/DataAssetFactory.h
--- a/Engine/Source/Editor/UnrealEd/Classes/Factories/DataAssetFactory.h
+++ b/Engine/Source/Editor/UnrealEd/Classes/Factories/DataAssetFactory.h
@@ -3,7 +3,7 @@
 #pragma once
 #include "DataAssetFactory.generated.h"
 
-UCLASS(hidecategories=Object)
+UCLASS(hidecategories=Object, MinimalAPI)
 class UDataAssetFactory : public UFactory
 {
 	GENERATED_UCLASS_BODY()
diff --git a/Engine/Source/Editor/UnrealEd/Private/Factories/EditorFactories.cpp b/Engine/Source/Editor/UnrealEd/Private/Factories/EditorFactories.cpp
--- a/Engine/Source/Editor/UnrealEd/Private/Factories/EditorFactories.cpp
+++ b/Engine/Source/Editor/UnrealEd/Private/Factories/EditorFactories.cpp
@@ -6643,7 +6643,9 @@ UObject* UDataAssetFactory::FactoryCreateNew(UClass* Class, UObject* InParent, F
 	}
 	else
 	{
-		return NULL;
+		// if we have no data asset class, use the passed-in class instead
+		check(Class->IsChildOf(UDataAsset::StaticClass()));
+		return CastChecked<UDataAsset>(StaticConstructObject(Class,InParent,Name,Flags));
 	}
 }
 
