Delete and rename Blackboard keys now work again from right-click menu

Added the ability to check for command registration (as this can crash if a command is registed more than once) and re-used the asset editor toolkit command list when registering commands.

TTP# 340628 - BINARY: EDITOR: Unable to rename/delete blackboard keys with right click menu

diff --git a/Engine/Source/Editor/BehaviorTreeEditor/Private/BehaviorTreeEditor.cpp b/Engine/Source/Editor/BehaviorTreeEditor/Private/BehaviorTreeEditor.cpp
--- a/Engine/Source/Editor/BehaviorTreeEditor/Private/BehaviorTreeEditor.cpp
+++ b/Engine/Source/Editor/BehaviorTreeEditor/Private/BehaviorTreeEditor.cpp
@@ -234,7 +234,7 @@ void FBehaviorTreeEditor::InitBehaviorTreeEditor( const EToolkitMode::Type Mode,
 		AddApplicationMode(BehaviorTreeMode, MakeShareable(new FBehaviorTreeEditorApplicationMode(SharedThis(this))));
 		AddApplicationMode(BlackboardMode, MakeShareable(new FBlackboardEditorApplicationMode(SharedThis(this))));
 
-		BlackboardView = SNew(SBehaviorTreeBlackboardView, GetBlackboardData())
+		BlackboardView = SNew(SBehaviorTreeBlackboardView, GetToolkitCommands(), GetBlackboardData())
 			.OnGetDebugKeyValue(this, &FBehaviorTreeEditor::HandleGetDebugKeyValue)
 			.OnIsDebuggerReady(this, &FBehaviorTreeEditor::IsDebuggerReady)
 			.OnIsDebuggerPaused(this, &FBehaviorTreeEditor::IsDebuggerPaused)
diff --git a/Engine/Source/Editor/BehaviorTreeEditor/Private/SBehaviorTreeBlackboardEditor.cpp b/Engine/Source/Editor/BehaviorTreeEditor/Private/SBehaviorTreeBlackboardEditor.cpp
--- a/Engine/Source/Editor/BehaviorTreeEditor/Private/SBehaviorTreeBlackboardEditor.cpp
+++ b/Engine/Source/Editor/BehaviorTreeEditor/Private/SBehaviorTreeBlackboardEditor.cpp
@@ -15,7 +15,7 @@
 
 #define LOCTEXT_NAMESPACE "SBehaviorTreeBlackboardEditor"
 
-void SBehaviorTreeBlackboardEditor::Construct(const FArguments& InArgs, TSharedRef<FUICommandList> ToolkitCommands, UBlackboardData* InBlackboardData)
+void SBehaviorTreeBlackboardEditor::Construct(const FArguments& InArgs, TSharedRef<FUICommandList> InToolkitCommands, UBlackboardData* InBlackboardData)
 {
 	OnEntrySelected = InArgs._OnEntrySelected;
 	OnGetDebugKeyValue = InArgs._OnGetDebugKeyValue;
@@ -24,17 +24,23 @@ void SBehaviorTreeBlackboardEditor::Construct(const FArguments& InArgs, TSharedR
 	OnGetDebugTimeStamp = InArgs._OnGetDebugTimeStamp;
 	OnGetDisplayCurrentState = InArgs._OnGetDisplayCurrentState;
 
-	ToolkitCommands->MapAction(
-		FBTBlackboardCommands::Get().DeleteEntry,
-		FExecuteAction::CreateSP(this, &SBehaviorTreeBlackboardEditor::HandleDeleteEntry),
-		FCanExecuteAction::CreateSP(this, &SBehaviorTreeBlackboardEditor::HasSelectedItems)
-		);
+	if(!InToolkitCommands->IsCommandInfoMapped(FBTBlackboardCommands::Get().DeleteEntry))
+	{
+		InToolkitCommands->MapAction(
+			FBTBlackboardCommands::Get().DeleteEntry,
+			FExecuteAction::CreateSP(this, &SBehaviorTreeBlackboardEditor::HandleDeleteEntry),
+			FCanExecuteAction::CreateSP(this, &SBehaviorTreeBlackboardEditor::HasSelectedItems)
+			);
+	}
 
-	ToolkitCommands->MapAction(
-		FGenericCommands::Get().Rename,
-		FExecuteAction::CreateSP(this, &SBehaviorTreeBlackboardEditor::HandleRenameEntry),
-		FCanExecuteAction::CreateSP(this, &SBehaviorTreeBlackboardEditor::HasSelectedItems)
-		);
+	if(!InToolkitCommands->IsCommandInfoMapped(FGenericCommands::Get().Rename))
+	{
+		InToolkitCommands->MapAction(
+			FGenericCommands::Get().Rename,
+			FExecuteAction::CreateSP(this, &SBehaviorTreeBlackboardEditor::HandleRenameEntry),
+			FCanExecuteAction::CreateSP(this, &SBehaviorTreeBlackboardEditor::HasSelectedItems)
+			);
+	}
 
 	SBehaviorTreeBlackboardView::Construct(
 		SBehaviorTreeBlackboardView::FArguments()
@@ -45,6 +51,7 @@ void SBehaviorTreeBlackboardEditor::Construct(const FArguments& InArgs, TSharedR
 		.OnIsDebuggerPaused(InArgs._OnIsDebuggerPaused)
 		.OnGetDebugTimeStamp(InArgs._OnGetDebugTimeStamp)
 		.IsReadOnly(false),
+		InToolkitCommands,
 		InBlackboardData
 	);
 }
diff --git a/Engine/Source/Editor/BehaviorTreeEditor/Private/SBehaviorTreeBlackboardView.cpp b/Engine/Source/Editor/BehaviorTreeEditor/Private/SBehaviorTreeBlackboardView.cpp
--- a/Engine/Source/Editor/BehaviorTreeEditor/Private/SBehaviorTreeBlackboardView.cpp
+++ b/Engine/Source/Editor/BehaviorTreeEditor/Private/SBehaviorTreeBlackboardView.cpp
@@ -252,7 +252,7 @@ void SBehaviorTreeBlackboardView::AddReferencedObjects( FReferenceCollector& Col
 	}
 }
 
-void SBehaviorTreeBlackboardView::Construct(const FArguments& InArgs, UBlackboardData* InBlackboardData)
+void SBehaviorTreeBlackboardView::Construct(const FArguments& InArgs, TSharedRef<FUICommandList> InToolkitCommands, UBlackboardData* InBlackboardData)
 {
 	OnEntrySelected = InArgs._OnEntrySelected;
 	OnGetDebugKeyValue = InArgs._OnGetDebugKeyValue;
@@ -265,29 +265,34 @@ void SBehaviorTreeBlackboardView::Construct(const FArguments& InArgs, UBlackboar
 
 	bShowCurrentState = OnGetDisplayCurrentState.IsBound() ? OnGetDisplayCurrentState.Execute() : true;
 
-	TSharedRef<FUICommandList> ToolkitCommands = MakeShareable(new FUICommandList);
-
-	ToolkitCommands->MapAction(
-		FBTDebuggerCommands::Get().CurrentValues,
-		FUIAction(
-			FExecuteAction::CreateSP(this, &SBehaviorTreeBlackboardView::HandleUseCurrentValues),
-			FCanExecuteAction::CreateSP(this, &SBehaviorTreeBlackboardView::IsDebuggerPaused),
-			FIsActionChecked::CreateSP(this, &SBehaviorTreeBlackboardView::IsUsingCurrentValues),
-			FIsActionButtonVisible::CreateSP(this, &SBehaviorTreeBlackboardView::IsDebuggerActive)
-			)
-		);
-	ToolkitCommands->MapAction(
-		FBTDebuggerCommands::Get().SavedValues, 
-		FUIAction(
-			FExecuteAction::CreateSP(this, &SBehaviorTreeBlackboardView::HandleUseSavedValues),
-			FCanExecuteAction::CreateSP(this, &SBehaviorTreeBlackboardView::IsDebuggerPaused),
-			FIsActionChecked::CreateSP(this, &SBehaviorTreeBlackboardView::IsUsingSavedValues),
-			FIsActionButtonVisible::CreateSP(this, &SBehaviorTreeBlackboardView::IsDebuggerActive)
-			)
-		);
+	if(!InToolkitCommands->IsCommandInfoMapped(FBTDebuggerCommands::Get().CurrentValues))
+	{
+		InToolkitCommands->MapAction(
+			FBTDebuggerCommands::Get().CurrentValues,
+			FUIAction(
+				FExecuteAction::CreateSP(this, &SBehaviorTreeBlackboardView::HandleUseCurrentValues),
+				FCanExecuteAction::CreateSP(this, &SBehaviorTreeBlackboardView::IsDebuggerPaused),
+				FIsActionChecked::CreateSP(this, &SBehaviorTreeBlackboardView::IsUsingCurrentValues),
+				FIsActionButtonVisible::CreateSP(this, &SBehaviorTreeBlackboardView::IsDebuggerActive)
+				)
+			);
+	}
+
+	if(!InToolkitCommands->IsCommandInfoMapped(FBTDebuggerCommands::Get().SavedValues))
+	{
+		InToolkitCommands->MapAction(
+			FBTDebuggerCommands::Get().SavedValues, 
+			FUIAction(
+				FExecuteAction::CreateSP(this, &SBehaviorTreeBlackboardView::HandleUseSavedValues),
+				FCanExecuteAction::CreateSP(this, &SBehaviorTreeBlackboardView::IsDebuggerPaused),
+				FIsActionChecked::CreateSP(this, &SBehaviorTreeBlackboardView::IsUsingSavedValues),
+				FIsActionButtonVisible::CreateSP(this, &SBehaviorTreeBlackboardView::IsDebuggerActive)
+				)
+			);
+	}
 
 	// build debug toolbar
-	FToolBarBuilder ToolbarBuilder(ToolkitCommands, FMultiBoxCustomization::None, GetToolbarExtender(ToolkitCommands));
+	FToolBarBuilder ToolbarBuilder(InToolkitCommands, FMultiBoxCustomization::None, GetToolbarExtender(InToolkitCommands));
 	
 	ToolbarBuilder.BeginSection(TEXT("Debugging"));
 	{
@@ -330,7 +335,7 @@ void SBehaviorTreeBlackboardView::Construct(const FArguments& InArgs, UBlackboar
 				.OnCollectAllActions(this, &SBehaviorTreeBlackboardView::HandleCollectAllActions)
 				.OnGetSectionTitle(this, &SBehaviorTreeBlackboardView::HandleGetSectionTitle)
 				.OnActionSelected(this, &SBehaviorTreeBlackboardView::HandleActionSelected)
-				.OnContextMenuOpening(this, &SBehaviorTreeBlackboardView::HandleContextMenuOpening, ToolkitCommands)
+				.OnContextMenuOpening(this, &SBehaviorTreeBlackboardView::HandleContextMenuOpening, InToolkitCommands)
 				.OnActionMatchesName(this, &SBehaviorTreeBlackboardView::HandleActionMatchesName)
 				.AlphaSortItems(false)
 			]
diff --git a/Engine/Source/Editor/BehaviorTreeEditor/Private/SBehaviorTreeBlackboardView.h b/Engine/Source/Editor/BehaviorTreeEditor/Private/SBehaviorTreeBlackboardView.h
--- a/Engine/Source/Editor/BehaviorTreeEditor/Private/SBehaviorTreeBlackboardView.h
+++ b/Engine/Source/Editor/BehaviorTreeEditor/Private/SBehaviorTreeBlackboardView.h
@@ -47,7 +47,7 @@ public:
 
 	SLATE_END_ARGS()
 
-	void Construct(const FArguments& InArgs, UBlackboardData* InBlackboardData);
+	void Construct(const FArguments& InArgs, TSharedRef<FUICommandList> InToolkitCommands, UBlackboardData* InBlackboardData);
 
 	/** FGCObject implementation */
 	virtual void AddReferencedObjects( FReferenceCollector& Collector ) override;
diff --git a/Engine/Source/Runtime/Slate/Private/Framework/Commands/UICommandList.cpp b/Engine/Source/Runtime/Slate/Private/Framework/Commands/UICommandList.cpp
--- a/Engine/Source/Runtime/Slate/Private/Framework/Commands/UICommandList.cpp
+++ b/Engine/Source/Runtime/Slate/Private/Framework/Commands/UICommandList.cpp
@@ -73,6 +73,12 @@ void FUICommandList::MapAction( const TSharedPtr< const FUICommandInfo > InUICom
 	UICommandBindingMap.Add( InUICommandInfo, InUIAction );
 }
 
+bool FUICommandList::IsCommandInfoMapped( const TSharedPtr< const FUICommandInfo > InUICommandInfo ) const
+{
+	check( InUICommandInfo.IsValid() );
+
+	return UICommandBindingMap.Contains( InUICommandInfo );
+}
 
 void FUICommandList::Append( const TSharedRef<FUICommandList>& InCommandsToAppend )
 {
diff --git a/Engine/Source/Runtime/Slate/Public/Framework/Commands/UICommandList.h b/Engine/Source/Runtime/Slate/Public/Framework/Commands/UICommandList.h
--- a/Engine/Source/Runtime/Slate/Public/Framework/Commands/UICommandList.h
+++ b/Engine/Source/Runtime/Slate/Public/Framework/Commands/UICommandList.h
@@ -61,6 +61,14 @@ public:
 	 */
 	void MapAction( const TSharedPtr< const FUICommandInfo > InUICommandInfo, const FUIAction& InUIAction );
 
+	/**
+	 * Checks to see whether a command info is already mapped to an action
+	 * 
+	 * @param InUICommandInfo	The command info to check
+	 * @return true if the command info is mapped to an action
+	 */
+	bool IsCommandInfoMapped( const TSharedPtr< const FUICommandInfo > InUICommandInfo ) const;
+
 	/**
 	 * Append commands in InCommandsToAppend to this command list.
 	 */
