Proper fix for Blackboard context menu

Reverted workaround in FUIComandList.
Properly appended command lists in BT editor so keyboard shortcut routing and context menus all work together.

TTP# 340628 - BINARY: EDITOR: Unable to rename/delete blackboard keys with right click menu

diff --git a/Engine/Source/Editor/BehaviorTreeEditor/Private/SBehaviorTreeBlackboardEditor.cpp b/Engine/Source/Editor/BehaviorTreeEditor/Private/SBehaviorTreeBlackboardEditor.cpp
--- a/Engine/Source/Editor/BehaviorTreeEditor/Private/SBehaviorTreeBlackboardEditor.cpp
+++ b/Engine/Source/Editor/BehaviorTreeEditor/Private/SBehaviorTreeBlackboardEditor.cpp
@@ -15,7 +15,7 @@
 
 #define LOCTEXT_NAMESPACE "SBehaviorTreeBlackboardEditor"
 
-void SBehaviorTreeBlackboardEditor::Construct(const FArguments& InArgs, TSharedRef<FUICommandList> InToolkitCommands, UBlackboardData* InBlackboardData)
+void SBehaviorTreeBlackboardEditor::Construct(const FArguments& InArgs, TSharedRef<FUICommandList> InCommandList, UBlackboardData* InBlackboardData)
 {
 	OnEntrySelected = InArgs._OnEntrySelected;
 	OnGetDebugKeyValue = InArgs._OnGetDebugKeyValue;
@@ -24,23 +24,21 @@ void SBehaviorTreeBlackboardEditor::Construct(const FArguments& InArgs, TSharedR
 	OnGetDebugTimeStamp = InArgs._OnGetDebugTimeStamp;
 	OnGetDisplayCurrentState = InArgs._OnGetDisplayCurrentState;
 
-	if(!InToolkitCommands->IsCommandInfoMapped(FBTBlackboardCommands::Get().DeleteEntry))
-	{
-		InToolkitCommands->MapAction(
-			FBTBlackboardCommands::Get().DeleteEntry,
-			FExecuteAction::CreateSP(this, &SBehaviorTreeBlackboardEditor::HandleDeleteEntry),
-			FCanExecuteAction::CreateSP(this, &SBehaviorTreeBlackboardEditor::HasSelectedItems)
-			);
-	}
+	TSharedRef<FUICommandList> CommandList = MakeShareable(new FUICommandList);
 
-	if(!InToolkitCommands->IsCommandInfoMapped(FGenericCommands::Get().Rename))
-	{
-		InToolkitCommands->MapAction(
-			FGenericCommands::Get().Rename,
-			FExecuteAction::CreateSP(this, &SBehaviorTreeBlackboardEditor::HandleRenameEntry),
-			FCanExecuteAction::CreateSP(this, &SBehaviorTreeBlackboardEditor::HasSelectedItems)
-			);
-	}
+	CommandList->MapAction(
+		FBTBlackboardCommands::Get().DeleteEntry,
+		FExecuteAction::CreateSP(this, &SBehaviorTreeBlackboardEditor::HandleDeleteEntry),
+		FCanExecuteAction::CreateSP(this, &SBehaviorTreeBlackboardEditor::HasSelectedItems)
+		);
+
+	CommandList->MapAction(
+		FGenericCommands::Get().Rename,
+		FExecuteAction::CreateSP(this, &SBehaviorTreeBlackboardEditor::HandleRenameEntry),
+		FCanExecuteAction::CreateSP(this, &SBehaviorTreeBlackboardEditor::HasSelectedItems)
+		);
+
+	InCommandList->Append(CommandList);
 
 	SBehaviorTreeBlackboardView::Construct(
 		SBehaviorTreeBlackboardView::FArguments()
@@ -51,7 +49,7 @@ void SBehaviorTreeBlackboardEditor::Construct(const FArguments& InArgs, TSharedR
 		.OnIsDebuggerPaused(InArgs._OnIsDebuggerPaused)
 		.OnGetDebugTimeStamp(InArgs._OnGetDebugTimeStamp)
 		.IsReadOnly(false),
-		InToolkitCommands,
+		CommandList,
 		InBlackboardData
 	);
 }
diff --git a/Engine/Source/Editor/BehaviorTreeEditor/Private/SBehaviorTreeBlackboardEditor.h b/Engine/Source/Editor/BehaviorTreeEditor/Private/SBehaviorTreeBlackboardEditor.h
--- a/Engine/Source/Editor/BehaviorTreeEditor/Private/SBehaviorTreeBlackboardEditor.h
+++ b/Engine/Source/Editor/BehaviorTreeEditor/Private/SBehaviorTreeBlackboardEditor.h
@@ -21,7 +21,7 @@ public:
 
 	SLATE_END_ARGS()
 
-	void Construct(const FArguments& InArgs, TSharedRef<FUICommandList> ToolkitCommands, UBlackboardData* InBlackboardData);
+	void Construct(const FArguments& InArgs, TSharedRef<FUICommandList> InCommandList, UBlackboardData* InBlackboardData);
 
 private:
 	/** Fill the context menu with edit options */
diff --git a/Engine/Source/Editor/BehaviorTreeEditor/Private/SBehaviorTreeBlackboardView.cpp b/Engine/Source/Editor/BehaviorTreeEditor/Private/SBehaviorTreeBlackboardView.cpp
--- a/Engine/Source/Editor/BehaviorTreeEditor/Private/SBehaviorTreeBlackboardView.cpp
+++ b/Engine/Source/Editor/BehaviorTreeEditor/Private/SBehaviorTreeBlackboardView.cpp
@@ -252,7 +252,7 @@ void SBehaviorTreeBlackboardView::AddReferencedObjects( FReferenceCollector& Col
 	}
 }
 
-void SBehaviorTreeBlackboardView::Construct(const FArguments& InArgs, TSharedRef<FUICommandList> InToolkitCommands, UBlackboardData* InBlackboardData)
+void SBehaviorTreeBlackboardView::Construct(const FArguments& InArgs, TSharedRef<FUICommandList> InCommandList, UBlackboardData* InBlackboardData)
 {
 	OnEntrySelected = InArgs._OnEntrySelected;
 	OnGetDebugKeyValue = InArgs._OnGetDebugKeyValue;
@@ -265,34 +265,32 @@ void SBehaviorTreeBlackboardView::Construct(const FArguments& InArgs, TSharedRef
 
 	bShowCurrentState = OnGetDisplayCurrentState.IsBound() ? OnGetDisplayCurrentState.Execute() : true;
 
-	if(!InToolkitCommands->IsCommandInfoMapped(FBTDebuggerCommands::Get().CurrentValues))
-	{
-		InToolkitCommands->MapAction(
-			FBTDebuggerCommands::Get().CurrentValues,
-			FUIAction(
-				FExecuteAction::CreateSP(this, &SBehaviorTreeBlackboardView::HandleUseCurrentValues),
-				FCanExecuteAction::CreateSP(this, &SBehaviorTreeBlackboardView::IsDebuggerPaused),
-				FIsActionChecked::CreateSP(this, &SBehaviorTreeBlackboardView::IsUsingCurrentValues),
-				FIsActionButtonVisible::CreateSP(this, &SBehaviorTreeBlackboardView::IsDebuggerActive)
-				)
-			);
-	}
-
-	if(!InToolkitCommands->IsCommandInfoMapped(FBTDebuggerCommands::Get().SavedValues))
-	{
-		InToolkitCommands->MapAction(
-			FBTDebuggerCommands::Get().SavedValues, 
-			FUIAction(
-				FExecuteAction::CreateSP(this, &SBehaviorTreeBlackboardView::HandleUseSavedValues),
-				FCanExecuteAction::CreateSP(this, &SBehaviorTreeBlackboardView::IsDebuggerPaused),
-				FIsActionChecked::CreateSP(this, &SBehaviorTreeBlackboardView::IsUsingSavedValues),
-				FIsActionButtonVisible::CreateSP(this, &SBehaviorTreeBlackboardView::IsDebuggerActive)
-				)
-			);
-	}
+	TSharedRef<FUICommandList> CommandList = MakeShareable(new FUICommandList);
+
+	CommandList->MapAction(
+		FBTDebuggerCommands::Get().CurrentValues,
+		FUIAction(
+			FExecuteAction::CreateSP(this, &SBehaviorTreeBlackboardView::HandleUseCurrentValues),
+			FCanExecuteAction::CreateSP(this, &SBehaviorTreeBlackboardView::IsDebuggerPaused),
+			FIsActionChecked::CreateSP(this, &SBehaviorTreeBlackboardView::IsUsingCurrentValues),
+			FIsActionButtonVisible::CreateSP(this, &SBehaviorTreeBlackboardView::IsDebuggerActive)
+			)
+		);
+
+	CommandList->MapAction(
+		FBTDebuggerCommands::Get().SavedValues, 
+		FUIAction(
+			FExecuteAction::CreateSP(this, &SBehaviorTreeBlackboardView::HandleUseSavedValues),
+			FCanExecuteAction::CreateSP(this, &SBehaviorTreeBlackboardView::IsDebuggerPaused),
+			FIsActionChecked::CreateSP(this, &SBehaviorTreeBlackboardView::IsUsingSavedValues),
+			FIsActionButtonVisible::CreateSP(this, &SBehaviorTreeBlackboardView::IsDebuggerActive)
+			)
+		);
+
+	InCommandList->Append(CommandList);
 
 	// build debug toolbar
-	FToolBarBuilder ToolbarBuilder(InToolkitCommands, FMultiBoxCustomization::None, GetToolbarExtender(InToolkitCommands));
+	FToolBarBuilder ToolbarBuilder(CommandList, FMultiBoxCustomization::None, GetToolbarExtender(InCommandList));
 	
 	ToolbarBuilder.BeginSection(TEXT("Debugging"));
 	{
@@ -335,7 +333,7 @@ void SBehaviorTreeBlackboardView::Construct(const FArguments& InArgs, TSharedRef
 				.OnCollectAllActions(this, &SBehaviorTreeBlackboardView::HandleCollectAllActions)
 				.OnGetSectionTitle(this, &SBehaviorTreeBlackboardView::HandleGetSectionTitle)
 				.OnActionSelected(this, &SBehaviorTreeBlackboardView::HandleActionSelected)
-				.OnContextMenuOpening(this, &SBehaviorTreeBlackboardView::HandleContextMenuOpening, InToolkitCommands)
+				.OnContextMenuOpening(this, &SBehaviorTreeBlackboardView::HandleContextMenuOpening, InCommandList)
 				.OnActionMatchesName(this, &SBehaviorTreeBlackboardView::HandleActionMatchesName)
 				.AlphaSortItems(false)
 			]
diff --git a/Engine/Source/Editor/BehaviorTreeEditor/Private/SBehaviorTreeBlackboardView.h b/Engine/Source/Editor/BehaviorTreeEditor/Private/SBehaviorTreeBlackboardView.h
--- a/Engine/Source/Editor/BehaviorTreeEditor/Private/SBehaviorTreeBlackboardView.h
+++ b/Engine/Source/Editor/BehaviorTreeEditor/Private/SBehaviorTreeBlackboardView.h
@@ -47,7 +47,7 @@ public:
 
 	SLATE_END_ARGS()
 
-	void Construct(const FArguments& InArgs, TSharedRef<FUICommandList> InToolkitCommands, UBlackboardData* InBlackboardData);
+	void Construct(const FArguments& InArgs, TSharedRef<FUICommandList> InCommandList, UBlackboardData* InBlackboardData);
 
 	/** FGCObject implementation */
 	virtual void AddReferencedObjects( FReferenceCollector& Collector ) override;
diff --git a/Engine/Source/Runtime/Slate/Private/Framework/Commands/UICommandList.cpp b/Engine/Source/Runtime/Slate/Private/Framework/Commands/UICommandList.cpp
--- a/Engine/Source/Runtime/Slate/Private/Framework/Commands/UICommandList.cpp
+++ b/Engine/Source/Runtime/Slate/Private/Framework/Commands/UICommandList.cpp
@@ -73,12 +73,6 @@ void FUICommandList::MapAction( const TSharedPtr< const FUICommandInfo > InUICom
 	UICommandBindingMap.Add( InUICommandInfo, InUIAction );
 }
 
-bool FUICommandList::IsCommandInfoMapped( const TSharedPtr< const FUICommandInfo > InUICommandInfo ) const
-{
-	check( InUICommandInfo.IsValid() );
-
-	return UICommandBindingMap.Contains( InUICommandInfo );
-}
 
 void FUICommandList::Append( const TSharedRef<FUICommandList>& InCommandsToAppend )
 {
diff --git a/Engine/Source/Runtime/Slate/Public/Framework/Commands/UICommandList.h b/Engine/Source/Runtime/Slate/Public/Framework/Commands/UICommandList.h
--- a/Engine/Source/Runtime/Slate/Public/Framework/Commands/UICommandList.h
+++ b/Engine/Source/Runtime/Slate/Public/Framework/Commands/UICommandList.h
@@ -61,14 +61,6 @@ public:
 	 */
 	void MapAction( const TSharedPtr< const FUICommandInfo > InUICommandInfo, const FUIAction& InUIAction );
 
-	/**
-	 * Checks to see whether a command info is already mapped to an action
-	 * 
-	 * @param InUICommandInfo	The command info to check
-	 * @return true if the command info is mapped to an action
-	 */
-	bool IsCommandInfoMapped( const TSharedPtr< const FUICommandInfo > InUICommandInfo ) const;
-
 	/**
 	 * Append commands in InCommandsToAppend to this command list.
 	 */
