Added Behaviour Tree node tooltips

Added new tooltip text using tooltip metadata from C++ node classes & BlueprintDescription text for BP node classes.

TTP# 339734 - Editor: Behaviour Tree: Add node tooltips

diff --git a/Engine/Source/Editor/BehaviorTreeEditor/Classes/BehaviorTreeGraphNode_Root.h b/Engine/Source/Editor/BehaviorTreeEditor/Classes/BehaviorTreeGraphNode_Root.h
--- a/Engine/Source/Editor/BehaviorTreeEditor/Classes/BehaviorTreeGraphNode_Root.h
+++ b/Engine/Source/Editor/BehaviorTreeEditor/Classes/BehaviorTreeGraphNode_Root.h
@@ -3,6 +3,7 @@
 #pragma once
 #include "BehaviorTreeGraphNode_Root.generated.h"
 
+/** Root node of this behavior tree, holds Blackboard data */
 UCLASS()
 class UBehaviorTreeGraphNode_Root : public UBehaviorTreeGraphNode
 {
diff --git a/Engine/Source/Editor/BehaviorTreeEditor/Private/BehaviorTreeGraphNode.cpp b/Engine/Source/Editor/BehaviorTreeEditor/Private/BehaviorTreeGraphNode.cpp
--- a/Engine/Source/Editor/BehaviorTreeEditor/Private/BehaviorTreeGraphNode.cpp
+++ b/Engine/Source/Editor/BehaviorTreeEditor/Private/BehaviorTreeGraphNode.cpp
@@ -5,6 +5,10 @@
 #include "SGraphEditorActionMenu_BehaviorTree.h"
 #include "BlueprintNodeHelpers.h"
 #include "BehaviorTree/BehaviorTree.h"
+#include "AssetData.h"
+#include "BehaviorTree/Decorators/BTDecorator_BlueprintBase.h"
+#include "BehaviorTree/Services/BTService_BlueprintBase.h"
+#include "BehaviorTree/Tasks/BTTask_BlueprintBase.h"
 
 #define LOCTEXT_NAMESPACE "BehaviorTreeGraphNode"
 
@@ -125,9 +129,45 @@ FString	UBehaviorTreeGraphNode::GetDescription() const
 
 FString UBehaviorTreeGraphNode::GetTooltip() const
 {
-	FString TooltipDesc = !NodeInstance ? LOCTEXT("NodeClassError", "Class not found, make sure it's saved!").ToString() :
-		bHasObserverError ? LOCTEXT("ObserverError", "Observer has invalid abort setting!").ToString() :		
-		(DebuggerRuntimeDescription.Len() > 0) ? DebuggerRuntimeDescription : ErrorMessage;
+	FString TooltipDesc;
+
+	if(!NodeInstance)
+	{
+		TooltipDesc = LOCTEXT("NodeClassError", "Class not found, make sure it's saved!").ToString();
+	}
+	else
+	{
+		if(bHasObserverError)
+		{
+			TooltipDesc = LOCTEXT("ObserverError", "Observer has invalid abort setting!").ToString();
+		}
+		else if(DebuggerRuntimeDescription.Len() > 0)
+		{
+			TooltipDesc = DebuggerRuntimeDescription;
+		}
+		else if(ErrorMessage.Len() > 0)
+		{
+			TooltipDesc = ErrorMessage;
+		}
+		else
+		{
+			if( NodeInstance->GetClass()->IsChildOf(UBTDecorator_BlueprintBase::StaticClass()) ||
+				NodeInstance->GetClass()->IsChildOf(UBTService_BlueprintBase::StaticClass()) ||
+				NodeInstance->GetClass()->IsChildOf(UBTTask_BlueprintBase::StaticClass()) )
+			{
+				FAssetData AssetData(NodeInstance->GetClass()->ClassGeneratedBy);
+				const FString* Description = AssetData.TagsAndValues.Find(GET_MEMBER_NAME_CHECKED(UBlueprint, BlueprintDescription));
+				if (Description && !Description->IsEmpty())
+				{
+					TooltipDesc = Description->Replace(TEXT("\\n"), TEXT("\n"));
+				}
+			}
+			else
+			{
+				TooltipDesc = NodeInstance->GetClass()->GetToolTipText().ToString();
+			}	
+		}
+	}
 
 	if (bInjectedNode)
 	{
diff --git a/Engine/Source/Editor/BehaviorTreeEditor/Private/BehaviorTreeGraphNode_Root.cpp b/Engine/Source/Editor/BehaviorTreeEditor/Private/BehaviorTreeGraphNode_Root.cpp
--- a/Engine/Source/Editor/BehaviorTreeEditor/Private/BehaviorTreeGraphNode_Root.cpp
+++ b/Engine/Source/Editor/BehaviorTreeEditor/Private/BehaviorTreeGraphNode_Root.cpp
@@ -42,7 +42,7 @@ FName UBehaviorTreeGraphNode_Root::GetNameIcon() const
 
 FString UBehaviorTreeGraphNode_Root::GetTooltip() const
 {
-	return NSLOCTEXT("BehaviorTreeGraphNode", "RootTooltip", "Root node of this behavior tree, holds Blackboard data").ToString();
+	return UEdGraphNode::GetTooltip();
 }
 
 void UBehaviorTreeGraphNode_Root::PostEditChangeProperty( struct FPropertyChangedEvent& PropertyChangedEvent)
diff --git a/Engine/Source/Runtime/AIModule/Classes/BehaviorTree/Composites/BTComposite_Selector.h b/Engine/Source/Runtime/AIModule/Classes/BehaviorTree/Composites/BTComposite_Selector.h
--- a/Engine/Source/Runtime/AIModule/Classes/BehaviorTree/Composites/BTComposite_Selector.h
+++ b/Engine/Source/Runtime/AIModule/Classes/BehaviorTree/Composites/BTComposite_Selector.h
@@ -4,6 +4,10 @@
 #include "BehaviorTree/BTCompositeNode.h"
 #include "BTComposite_Selector.generated.h"
 
+/** 
+ * Selector composite node.
+ * A selector node runs each child in turn until the first child node succeeds, in which case the selector fails.
+ */
 UCLASS()
 class UBTComposite_Selector : public UBTCompositeNode
 {
diff --git a/Engine/Source/Runtime/AIModule/Classes/BehaviorTree/Composites/BTComposite_Sequence.h b/Engine/Source/Runtime/AIModule/Classes/BehaviorTree/Composites/BTComposite_Sequence.h
--- a/Engine/Source/Runtime/AIModule/Classes/BehaviorTree/Composites/BTComposite_Sequence.h
+++ b/Engine/Source/Runtime/AIModule/Classes/BehaviorTree/Composites/BTComposite_Sequence.h
@@ -4,6 +4,10 @@
 #include "BehaviorTree/BTCompositeNode.h"
 #include "BTComposite_Sequence.generated.h"
 
+/** 
+ * Sequence composite node.
+ * A sequence node runs each child in turn until one fails, in which case the sequence succeeds.
+ */
 UCLASS()
 class UBTComposite_Sequence : public UBTCompositeNode
 {
diff --git a/Engine/Source/Runtime/AIModule/Classes/BehaviorTree/Composites/BTComposite_SimpleParallel.h b/Engine/Source/Runtime/AIModule/Classes/BehaviorTree/Composites/BTComposite_SimpleParallel.h
--- a/Engine/Source/Runtime/AIModule/Classes/BehaviorTree/Composites/BTComposite_SimpleParallel.h
+++ b/Engine/Source/Runtime/AIModule/Classes/BehaviorTree/Composites/BTComposite_SimpleParallel.h
@@ -37,6 +37,10 @@ struct FBTParallelMemory : public FBTCompositeMemory
 	uint8 bForceBackgroundTree : 1;
 };
 
+/**
+ * Simple Parallel composite node.
+ * Allows for running two children: one which must be a single task node (with optional decorators), and the other of which can be a complete subtree.
+ */
 UCLASS(MinimalAPI)
 class UBTComposite_SimpleParallel : public UBTCompositeNode
 {
diff --git a/Engine/Source/Runtime/AIModule/Classes/BehaviorTree/Decorators/BTDecorator_Blackboard.h b/Engine/Source/Runtime/AIModule/Classes/BehaviorTree/Decorators/BTDecorator_Blackboard.h
--- a/Engine/Source/Runtime/AIModule/Classes/BehaviorTree/Decorators/BTDecorator_Blackboard.h
+++ b/Engine/Source/Runtime/AIModule/Classes/BehaviorTree/Decorators/BTDecorator_Blackboard.h
@@ -19,6 +19,10 @@ namespace EBTBlackboardRestart
 	};
 }
 
+/**
+ * Blackboard decorator node.
+ * A decorator node that bases its condition on a Blackboard key.
+ */
 UCLASS(HideCategories=(Condition))
 class AIMODULE_API UBTDecorator_Blackboard : public UBTDecorator_BlackboardBase
 {
diff --git a/Engine/Source/Runtime/AIModule/Classes/BehaviorTree/Decorators/BTDecorator_CompareBBEntries.h b/Engine/Source/Runtime/AIModule/Classes/BehaviorTree/Decorators/BTDecorator_CompareBBEntries.h
--- a/Engine/Source/Runtime/AIModule/Classes/BehaviorTree/Decorators/BTDecorator_CompareBBEntries.h
+++ b/Engine/Source/Runtime/AIModule/Classes/BehaviorTree/Decorators/BTDecorator_CompareBBEntries.h
@@ -14,6 +14,10 @@ namespace EBlackBoardEntryComparison
 	};
 }
 
+/**
+ * Blackboard comparison decorator node.
+ * A decorator node that bases its condition on a comparison between two Blackboard keys.
+ */
 UCLASS(HideCategories=(Condition))
 class UBTDecorator_CompareBBEntries : public UBTDecorator
 {
diff --git a/Engine/Source/Runtime/AIModule/Classes/BehaviorTree/Decorators/BTDecorator_ConeCheck.h b/Engine/Source/Runtime/AIModule/Classes/BehaviorTree/Decorators/BTDecorator_ConeCheck.h
--- a/Engine/Source/Runtime/AIModule/Classes/BehaviorTree/Decorators/BTDecorator_ConeCheck.h
+++ b/Engine/Source/Runtime/AIModule/Classes/BehaviorTree/Decorators/BTDecorator_ConeCheck.h
@@ -9,6 +9,10 @@ struct FBTConeCheckDecoratorMemory
 	bool bLastRawResult;
 };
 
+/**
+ * Cone check decorator node.
+ * A decorator node that bases its condition on a cone check, using Blackboard entries to form the parameters of the check.
+ */
 UCLASS()
 class AIMODULE_API UBTDecorator_ConeCheck : public UBTDecorator
 {
diff --git a/Engine/Source/Runtime/AIModule/Classes/BehaviorTree/Decorators/BTDecorator_Cooldown.h b/Engine/Source/Runtime/AIModule/Classes/BehaviorTree/Decorators/BTDecorator_Cooldown.h
--- a/Engine/Source/Runtime/AIModule/Classes/BehaviorTree/Decorators/BTDecorator_Cooldown.h
+++ b/Engine/Source/Runtime/AIModule/Classes/BehaviorTree/Decorators/BTDecorator_Cooldown.h
@@ -10,6 +10,10 @@ struct FBTCooldownDecoratorMemory
 	uint8 RequestedRestart;
 };
 
+/**
+ * Cooldown decorator node.
+ * A decorator node that bases its condition on whether a cooldown timer has expired.
+ */
 UCLASS(HideCategories=(Condition))
 class AIMODULE_API UBTDecorator_Cooldown : public UBTDecorator
 {
diff --git a/Engine/Source/Runtime/AIModule/Classes/BehaviorTree/Decorators/BTDecorator_DoesPathExist.h b/Engine/Source/Runtime/AIModule/Classes/BehaviorTree/Decorators/BTDecorator_DoesPathExist.h
--- a/Engine/Source/Runtime/AIModule/Classes/BehaviorTree/Decorators/BTDecorator_DoesPathExist.h
+++ b/Engine/Source/Runtime/AIModule/Classes/BehaviorTree/Decorators/BTDecorator_DoesPathExist.h
@@ -15,6 +15,10 @@ namespace EPathExistanceQueryType
 	};
 }
 
+/**
+ * Cooldown decorator node.
+ * A decorator node that bases its condition on whether a path exists between two points from the Blackboard.
+ */
 UCLASS()
 class AIMODULE_API UBTDecorator_DoesPathExist : public UBTDecorator
 {
diff --git a/Engine/Source/Runtime/AIModule/Classes/BehaviorTree/Decorators/BTDecorator_KeepInCone.h b/Engine/Source/Runtime/AIModule/Classes/BehaviorTree/Decorators/BTDecorator_KeepInCone.h
--- a/Engine/Source/Runtime/AIModule/Classes/BehaviorTree/Decorators/BTDecorator_KeepInCone.h
+++ b/Engine/Source/Runtime/AIModule/Classes/BehaviorTree/Decorators/BTDecorator_KeepInCone.h
@@ -9,6 +9,10 @@ struct FBTKeepInConeDecoratorMemory
 	FVector InitialDirection;
 };
 
+/**
+ * Cooldown decorator node.
+ * A decorator node that bases its condition on whether the observed position is still inside a cone. The cone's direction is calculated when the node first becomes relevant.
+ */
 UCLASS(HideCategories=(Condition))
 class AIMODULE_API UBTDecorator_KeepInCone : public UBTDecorator
 {
diff --git a/Engine/Source/Runtime/AIModule/Classes/BehaviorTree/Decorators/BTDecorator_Loop.h b/Engine/Source/Runtime/AIModule/Classes/BehaviorTree/Decorators/BTDecorator_Loop.h
--- a/Engine/Source/Runtime/AIModule/Classes/BehaviorTree/Decorators/BTDecorator_Loop.h
+++ b/Engine/Source/Runtime/AIModule/Classes/BehaviorTree/Decorators/BTDecorator_Loop.h
@@ -8,6 +8,10 @@ struct FBTLoopDecoratorMemory
 	uint8 RemainingExecutions;
 };
 
+/**
+ * Loop decorator node.
+ * A decorator node that bases its condition on whether its loop counter has been exceeded.
+ */
 UCLASS(HideCategories=(Condition))
 class AIMODULE_API UBTDecorator_Loop : public UBTDecorator
 {
diff --git a/Engine/Source/Runtime/AIModule/Classes/BehaviorTree/Decorators/BTDecorator_ReachedMoveGoal.h b/Engine/Source/Runtime/AIModule/Classes/BehaviorTree/Decorators/BTDecorator_ReachedMoveGoal.h
--- a/Engine/Source/Runtime/AIModule/Classes/BehaviorTree/Decorators/BTDecorator_ReachedMoveGoal.h
+++ b/Engine/Source/Runtime/AIModule/Classes/BehaviorTree/Decorators/BTDecorator_ReachedMoveGoal.h
@@ -4,6 +4,10 @@
 #include "BehaviorTree/BTDecorator.h"
 #include "BTDecorator_ReachedMoveGoal.generated.h"
 
+/**
+ * Reached Move Goal decorator node.
+ * A decorator node that bases its condition on whether the AI controller's path following component returns that it has reached its goal.
+ */
 UCLASS()
 class AIMODULE_API UBTDecorator_ReachedMoveGoal : public UBTDecorator
 {
diff --git a/Engine/Source/Runtime/AIModule/Classes/BehaviorTree/Decorators/BTDecorator_TimeLimit.h b/Engine/Source/Runtime/AIModule/Classes/BehaviorTree/Decorators/BTDecorator_TimeLimit.h
--- a/Engine/Source/Runtime/AIModule/Classes/BehaviorTree/Decorators/BTDecorator_TimeLimit.h
+++ b/Engine/Source/Runtime/AIModule/Classes/BehaviorTree/Decorators/BTDecorator_TimeLimit.h
@@ -3,6 +3,10 @@
 #include "BehaviorTree/BTDecorator.h"
 #include "BTDecorator_TimeLimit.generated.h"
 
+/**
+ * Time Limit decorator node.
+ * A decorator node that bases its condition on whether a timer has exceeded a specified value. The timer is reset each time the node becomes relevant.
+ */
 UCLASS(HideCategories=(Condition))
 class AIMODULE_API UBTDecorator_TimeLimit : public UBTDecorator
 {
diff --git a/Engine/Source/Runtime/AIModule/Classes/BehaviorTree/Services/BTService_DefaultFocus.h b/Engine/Source/Runtime/AIModule/Classes/BehaviorTree/Services/BTService_DefaultFocus.h
--- a/Engine/Source/Runtime/AIModule/Classes/BehaviorTree/Services/BTService_DefaultFocus.h
+++ b/Engine/Source/Runtime/AIModule/Classes/BehaviorTree/Services/BTService_DefaultFocus.h
@@ -16,6 +16,10 @@ struct FBTDefaultFocusMemory
 	}
 };
 
+/**
+ * Default Focus service node.
+ * A service node that automatically sets the AI controller's focus when it becomes active.
+ */
 UCLASS(hidecategories=(Service))
 class UBTService_DefaultFocus : public UBTService_BlackboardBase
 {
diff --git a/Engine/Source/Runtime/AIModule/Classes/BehaviorTree/Tasks/BTTask_MakeNoise.h b/Engine/Source/Runtime/AIModule/Classes/BehaviorTree/Tasks/BTTask_MakeNoise.h
--- a/Engine/Source/Runtime/AIModule/Classes/BehaviorTree/Tasks/BTTask_MakeNoise.h
+++ b/Engine/Source/Runtime/AIModule/Classes/BehaviorTree/Tasks/BTTask_MakeNoise.h
@@ -3,6 +3,10 @@
 #pragma once
 #include "BTTask_MakeNoise.generated.h"
 
+/**
+ * Make Noise task node.
+ * A task node that calls MakeNoise() on this Pawn when executed.
+ */
 UCLASS()
 class UBTTask_MakeNoise : public UBTTaskNode
 {
diff --git a/Engine/Source/Runtime/AIModule/Classes/BehaviorTree/Tasks/BTTask_MoveDirectlyToward.h b/Engine/Source/Runtime/AIModule/Classes/BehaviorTree/Tasks/BTTask_MoveDirectlyToward.h
--- a/Engine/Source/Runtime/AIModule/Classes/BehaviorTree/Tasks/BTTask_MoveDirectlyToward.h
+++ b/Engine/Source/Runtime/AIModule/Classes/BehaviorTree/Tasks/BTTask_MoveDirectlyToward.h
@@ -10,8 +10,11 @@ struct FBTMoveDirectlyTowardMemory
 	FAIRequestID MoveRequestID;
 };
 
-UCLASS(config=Game, Meta=(
-	Tooltip="Moves the AI pawn toward the specified Actor or Location (Vector) blackboard entry in a straight line, without regard to any navigation system.  If you need the AI to navigate, use the \"Move To\" node instead."))
+/**
+ * Move Directly Toward task node.
+ * Moves the AI pawn toward the specified Actor or Location (Vector) blackboard entry in a straight line, without regard to any navigation system. If you need the AI to navigate, use the "Move To" node instead.
+ */
+UCLASS(config=Game)
 class AIMODULE_API UBTTask_MoveDirectlyToward : public UBTTask_BlackboardBase
 {
 	GENERATED_UCLASS_BODY()
diff --git a/Engine/Source/Runtime/AIModule/Classes/BehaviorTree/Tasks/BTTask_MoveTo.h b/Engine/Source/Runtime/AIModule/Classes/BehaviorTree/Tasks/BTTask_MoveTo.h
--- a/Engine/Source/Runtime/AIModule/Classes/BehaviorTree/Tasks/BTTask_MoveTo.h
+++ b/Engine/Source/Runtime/AIModule/Classes/BehaviorTree/Tasks/BTTask_MoveTo.h
@@ -12,6 +12,10 @@ struct FBTMoveToTaskMemory
 	uint8 bWaitingForPath : 1;
 };
 
+/**
+ * Move To task node.
+ * Moves the AI pawn toward the specified Actor or Location blackboard entry using the navigation system.
+ */
 UCLASS(config=Game)
 class AIMODULE_API UBTTask_MoveTo : public UBTTask_BlackboardBase
 {
diff --git a/Engine/Source/Runtime/AIModule/Classes/BehaviorTree/Tasks/BTTask_PlaySound.h b/Engine/Source/Runtime/AIModule/Classes/BehaviorTree/Tasks/BTTask_PlaySound.h
--- a/Engine/Source/Runtime/AIModule/Classes/BehaviorTree/Tasks/BTTask_PlaySound.h
+++ b/Engine/Source/Runtime/AIModule/Classes/BehaviorTree/Tasks/BTTask_PlaySound.h
@@ -6,6 +6,10 @@
 class UBehaviorTreeComponent;
 class USoundCue;
 
+/**
+ * Play Sound task node.
+ * Plays the specified sound when executed.
+ */
 UCLASS()
 class UBTTask_PlaySound : public UBTTaskNode
 {
diff --git a/Engine/Source/Runtime/AIModule/Classes/BehaviorTree/Tasks/BTTask_RunEQSQuery.h b/Engine/Source/Runtime/AIModule/Classes/BehaviorTree/Tasks/BTTask_RunEQSQuery.h
--- a/Engine/Source/Runtime/AIModule/Classes/BehaviorTree/Tasks/BTTask_RunEQSQuery.h
+++ b/Engine/Source/Runtime/AIModule/Classes/BehaviorTree/Tasks/BTTask_RunEQSQuery.h
@@ -10,6 +10,10 @@ struct FBTEnvQueryTaskMemory
 	int32 RequestID;
 };
 
+/**
+ * Run Environment Query System Query task node.
+ * Runs the specified environment query when executed.
+ */
 UCLASS(MinimalAPI)
 class UBTTask_RunEQSQuery : public UBTTask_BlackboardBase
 {
diff --git a/Engine/Source/Runtime/AIModule/Classes/BehaviorTree/Tasks/BTTask_Wait.h b/Engine/Source/Runtime/AIModule/Classes/BehaviorTree/Tasks/BTTask_Wait.h
--- a/Engine/Source/Runtime/AIModule/Classes/BehaviorTree/Tasks/BTTask_Wait.h
+++ b/Engine/Source/Runtime/AIModule/Classes/BehaviorTree/Tasks/BTTask_Wait.h
@@ -9,6 +9,10 @@ struct FBTWaitTaskMemory
 	float RemainingWaitTime;
 };
 
+/**
+ * Wait task node.
+ * Wait for the specified time when executed.
+ */
 UCLASS()
 class UBTTask_Wait : public UBTTaskNode
 {
