Added Behavior Tree execution order index widgets

Added new widget overlay system for graph panels, designed to allow more flexibilty than the current brush-only implementation. This requred a small refactor to avoid the assumption that all child widgets of an SNodePanel are SNodes (and because there is no usable RTTI present).

Added index widget overlay to BT graph nodes.

Added flag to parent nodes to allow us to change the appearance of child nodes when another child is hovered.

TTP# 339733 - Editor: Behaviour Tree: Add better execution order indicators

diff --git a/Engine/Source/Editor/BehaviorTreeEditor/Classes/BehaviorTreeGraphNode.h b/Engine/Source/Editor/BehaviorTreeEditor/Classes/BehaviorTreeGraphNode.h
--- a/Engine/Source/Editor/BehaviorTreeEditor/Classes/BehaviorTreeGraphNode.h
+++ b/Engine/Source/Editor/BehaviorTreeEditor/Classes/BehaviorTreeGraphNode.h
@@ -118,6 +118,9 @@ class UBehaviorTreeGraphNode : public UEdGraphNode
 	/** highlighting nodes during quick find */
 	uint32 bHighlightInSearchTree : 1;
 
+	/** highlight other child node indexes when hovering over a child */
+	uint32 bHighlightChildNodeIndices : 1;
+
 	/** debugger flag: breakpoint exists */
 	uint32 bHasBreakpoint : 1;
 
diff --git a/Engine/Source/Editor/BehaviorTreeEditor/Private/SGraphNode_BehaviorTree.cpp b/Engine/Source/Editor/BehaviorTreeEditor/Private/SGraphNode_BehaviorTree.cpp
--- a/Engine/Source/Editor/BehaviorTreeEditor/Private/SGraphNode_BehaviorTree.cpp
+++ b/Engine/Source/Editor/BehaviorTreeEditor/Private/SGraphNode_BehaviorTree.cpp
@@ -14,6 +14,8 @@
 #include "BehaviorTree/BehaviorTree.h"
 #include "BehaviorTree/Tasks/BTTask_RunBehavior.h"
 
+#define LOCTEXT_NAMESPACE "SGraphNode_BehaviorTree"
+
 /////////////////////////////////////////////////////
 // SBehaviorTreePin
 
@@ -79,6 +81,91 @@ FSlateColor SBehaviorTreePin::GetPinColor() const
 		BehaviorTreeColors::Pin::Default;
 }
 
+/** Widget for overlaying an execution-order index onto a node */
+class SBehaviorTreeIndex : public SCompoundWidget
+{
+public:
+	/** Delegate event fired when the hover state of this widget changes */
+	DECLARE_DELEGATE_OneParam(FOnHoverStateChanged, bool /* bHovered */);
+
+	/** Delegate used to receive the color of the node, depending on hover state and state of other siblings */
+	DECLARE_DELEGATE_RetVal_OneParam(FSlateColor, FOnGetIndexColor, bool /* bHovered */);
+
+	SLATE_BEGIN_ARGS(SBehaviorTreeIndex){}
+		SLATE_ATTRIBUTE(FText, Text)
+		SLATE_EVENT(FOnHoverStateChanged, OnHoverStateChanged)
+		SLATE_EVENT(FOnGetIndexColor, OnGetIndexColor)
+	SLATE_END_ARGS()
+
+	void Construct( const FArguments& InArgs )
+	{
+		OnHoverStateChangedEvent = InArgs._OnHoverStateChanged;
+		OnGetIndexColorEvent = InArgs._OnGetIndexColor;
+
+		const FSlateBrush* IndexBrush = FEditorStyle::GetBrush(TEXT("BTEditor.Graph.BTNode.Index"));
+
+		ChildSlot
+		[
+			SNew(SOverlay)
+			+SOverlay::Slot()
+			.HAlign(HAlign_Fill)
+			.VAlign(VAlign_Fill)
+			[
+				// Add a dummy box here to make sure the widget doesnt get smaller than the brush
+				SNew(SBox)
+				.WidthOverride(IndexBrush->ImageSize.X)
+				.HeightOverride(IndexBrush->ImageSize.Y)
+			]
+			+SOverlay::Slot()
+			.HAlign(HAlign_Fill)
+			.VAlign(VAlign_Fill)
+			[
+				SNew(SBorder)
+				.BorderImage(IndexBrush)
+				.BorderBackgroundColor(this, &SBehaviorTreeIndex::GetColor)
+				.Padding(FMargin(4.0f, 0.0f, 4.0f, 1.0f))
+				.VAlign(VAlign_Center)
+				.HAlign(HAlign_Center)
+				[
+					SNew(STextBlock)
+					.Text(InArgs._Text)
+					.Font(FEditorStyle::GetFontStyle("BTEditor.Graph.BTNode.IndexText"))
+				]
+			]
+		];
+	}
+
+	virtual void OnMouseEnter( const FGeometry& MyGeometry, const FPointerEvent& MouseEvent ) override
+	{
+		OnHoverStateChangedEvent.ExecuteIfBound(true);
+		SCompoundWidget::OnMouseEnter(MyGeometry, MouseEvent);
+	}
+
+	virtual void OnMouseLeave( const FPointerEvent& MouseEvent ) override
+	{
+		OnHoverStateChangedEvent.ExecuteIfBound(false);
+		SCompoundWidget::OnMouseLeave(MouseEvent);
+	}
+
+	/** Get the color we use to display the rounded border */
+	FSlateColor GetColor() const
+	{
+		if(OnGetIndexColorEvent.IsBound())
+		{
+			return OnGetIndexColorEvent.Execute(IsHovered());
+		}
+
+		return FSlateColor::UseForeground();
+	}
+
+private:
+	/** Delegate event fired when the hover state of this widget changes */
+	FOnHoverStateChanged OnHoverStateChangedEvent;
+
+	/** Delegate used to receive the color of the node, depending on hover state and state of other siblings */
+	FOnGetIndexColor OnGetIndexColorEvent;
+};
+
 /////////////////////////////////////////////////////
 // SGraphNode_BehaviorTree
 
@@ -268,6 +355,13 @@ void SGraphNode_BehaviorTree::UpdateGraphNode()
 		NodePadding = 2.0f;
 	}
 
+	IndexOverlay = SNew(SBehaviorTreeIndex)
+		.ToolTipText(this, &SGraphNode_BehaviorTree::GetIndexTooltipText)
+		.Visibility(this, &SGraphNode_BehaviorTree::GetIndexVisibility)
+		.Text(this, &SGraphNode_BehaviorTree::GetIndexText)
+		.OnHoverStateChanged(this, &SGraphNode_BehaviorTree::OnIndexHoverStateChanged)
+		.OnGetIndexColor(this, &SGraphNode_BehaviorTree::GetIndexColor);
+
 	this->ContentScale.Bind( this, &SGraphNode::GetContentScale );
 	this->ChildSlot
 		.HAlign(HAlign_Fill)
@@ -321,7 +415,7 @@ void SGraphNode_BehaviorTree::UpdateGraphNode()
 					+SVerticalBox::Slot()
 					.AutoHeight()
 					[
-						SNew(SBorder)
+						SAssignNew(NodeBody, SBorder)
 						.BorderImage( FEditorStyle::GetBrush("BTEditor.Graph.BTNode.Body") )
 						.BorderBackgroundColor( this, &SGraphNode_BehaviorTree::GetBackgroundColor )
 						.HAlign(HAlign_Fill)
@@ -413,14 +507,14 @@ void SGraphNode_BehaviorTree::UpdateGraphNode()
 					SNew(SBorder)
 					.BorderBackgroundColor(BehaviorTreeColors::Action::DragMarker)
 					.ColorAndOpacity(BehaviorTreeColors::Action::DragMarker)
-					.BorderImage(FEditorStyle::GetBrush("Graph.BTNode.Body"))
+					.BorderImage(FEditorStyle::GetBrush("BTEditor.Graph.BTNode.Body"))
 					.Visibility(this, &SGraphNode_BehaviorTree::GetDragOverMarkerVisibility)
 					[
 						SNew(SBox)
 						.HeightOverride(4)
 					]
 				]
-				//drag marker overlay
+				// Blueprint indicator overlay
 				+SOverlay::Slot()
 				.HAlign(HAlign_Left)
 				.VAlign(VAlign_Top)
@@ -767,50 +861,7 @@ FReply SGraphNode_BehaviorTree::OnDrop( const FGeometry& MyGeometry, const FDrag
 FString	SGraphNode_BehaviorTree::GetDescription() const
 {
 	UBehaviorTreeGraphNode* StateNode = CastChecked<UBehaviorTreeGraphNode>(GraphNode);
-	UEdGraphPin* MyInputPin = StateNode->GetInputPin();
-	UEdGraphPin* MyParentOutputPin = NULL;
-	if (MyInputPin != NULL && MyInputPin->LinkedTo.Num() > 0)
-	{
-		MyParentOutputPin = MyInputPin->LinkedTo[0];
-	}
-
-	int32 Index = 0;
-
-	if (GEditor->bIsSimulatingInEditor || GEditor->PlayWorld != NULL)
-	{
-		// show execution index (debugging purposes)
-		UBTNode* BTNode = Cast<UBTNode>(StateNode->NodeInstance);
-		Index = (BTNode && BTNode->GetExecutionIndex() < 0xffff) ? BTNode->GetExecutionIndex() : -1;
-
-		// special case: range of execution indices in composite decorator node
-		UBehaviorTreeGraphNode_CompositeDecorator* CompDecorator = Cast<UBehaviorTreeGraphNode_CompositeDecorator>(GraphNode);
-		if (CompDecorator && CompDecorator->FirstExecutionIndex != CompDecorator->LastExecutionIndex)
-		{
-			return FString::Printf( TEXT("(%d..%d) %s"),
-				CompDecorator->FirstExecutionIndex, CompDecorator->LastExecutionIndex, *StateNode->GetDescription());
-		}
-	}
-	else
-	{
-		// show child index
-	if (MyParentOutputPin != NULL)
-	{
-		for (Index = 0; Index < MyParentOutputPin->LinkedTo.Num(); ++Index)
-		{
-			if (MyParentOutputPin->LinkedTo[Index] == MyInputPin)
-			{
-				break;
-			}
-		}
-	}
-
-	if (!MyParentOutputPin || MyParentOutputPin->LinkedTo.Num() < 2 )
-	{
-		return FString::Printf( TEXT("%s"), StateNode ? *StateNode->GetDescription() : TEXT(""));
-	}
-	}
-
-	return FString::Printf( TEXT("(%d) %s"), Index, StateNode ? *StateNode->GetDescription() : TEXT(""));
+	return StateNode->GetDescription();
 }
 
 FString SGraphNode_BehaviorTree::GetPinTooltip(UEdGraphPin* GraphPinObj) const
@@ -836,14 +887,14 @@ void SGraphNode_BehaviorTree::CreatePinWidgets()
 	{
 		UEdGraphPin* MyPin = StateNode->Pins[PinIdx];
 		if (!MyPin->bHidden)
-	{
+		{
 			TSharedPtr<SGraphPin> NewPin = SNew(SBehaviorTreePin, MyPin)
 				.ToolTipText( this, &SGraphNode_BehaviorTree::GetPinTooltip, MyPin);
 
-		NewPin->SetIsEditable(IsEditable);
+			NewPin->SetIsEditable(IsEditable);
 
 			AddPin(NewPin.ToSharedRef());
-	}
+		}
 	}
 }
 
@@ -926,6 +977,122 @@ const FSlateBrush* SGraphNode_BehaviorTree::GetNameIcon() const
 	return BTGraphNode != nullptr ? FEditorStyle::GetBrush(BTGraphNode->GetNameIcon()) : FEditorStyle::GetBrush(TEXT("BTEditor.Graph.BTNode.Icon"));
 }
 
+static UBehaviorTreeGraphNode* GetParentNode(UEdGraphNode* GraphNode)
+{
+	UBehaviorTreeGraphNode* StateNode = CastChecked<UBehaviorTreeGraphNode>(GraphNode);
+	if(StateNode->ParentNode != nullptr)
+	{
+		StateNode = StateNode->ParentNode;
+	}
+	UEdGraphPin* MyInputPin = StateNode->GetInputPin();
+	UEdGraphPin* MyParentOutputPin = nullptr;
+	if (MyInputPin != nullptr && MyInputPin->LinkedTo.Num() > 0)
+	{
+		MyParentOutputPin = MyInputPin->LinkedTo[0];
+		if(MyParentOutputPin != nullptr)
+		{
+			if(MyParentOutputPin->GetOwningNode() != nullptr)
+			{
+				return CastChecked<UBehaviorTreeGraphNode>(MyParentOutputPin->GetOwningNode());
+			}
+		}
+	}
+
+	return nullptr;
+}
+
+void SGraphNode_BehaviorTree::OnIndexHoverStateChanged(bool bHovered)
+{
+	UBehaviorTreeGraphNode* ParentNode = GetParentNode(GraphNode);
+	if(ParentNode != nullptr)
+	{
+		ParentNode->bHighlightChildNodeIndices = bHovered;
+	}
+}
+
+FSlateColor SGraphNode_BehaviorTree::GetIndexColor(bool bHovered) const
+{
+	UBehaviorTreeGraphNode* ParentNode = GetParentNode(GraphNode);
+	const bool bHighlightHover = bHovered || (ParentNode && ParentNode->bHighlightChildNodeIndices);
+	return bHighlightHover ? FEditorStyle::Get().GetSlateColor("BTEditor.Graph.BTNode.Index.HoveredColor") : FEditorStyle::Get().GetSlateColor("BTEditor.Graph.BTNode.Index.Color");
+}
+
+EVisibility SGraphNode_BehaviorTree::GetIndexVisibility() const
+{
+	// always hide the index on the root node
+	if(GraphNode->IsA(UBehaviorTreeGraphNode_Root::StaticClass()))
+	{
+		return EVisibility::Collapsed;
+	}
+
+	UBehaviorTreeGraphNode* StateNode = CastChecked<UBehaviorTreeGraphNode>(GraphNode);
+	UEdGraphPin* MyInputPin = StateNode->GetInputPin();
+	UEdGraphPin* MyParentOutputPin = NULL;
+	if (MyInputPin != NULL && MyInputPin->LinkedTo.Num() > 0)
+	{
+		MyParentOutputPin = MyInputPin->LinkedTo[0];
+	}
+
+	// Visible if we are in PIE or if we have siblings
+	const bool bIsVisible = (GEditor->bIsSimulatingInEditor || GEditor->PlayWorld != NULL) || (MyParentOutputPin && MyParentOutputPin->LinkedTo.Num() > 1);
+	return bIsVisible ? EVisibility::Visible : EVisibility::Collapsed;
+}
+
+FText SGraphNode_BehaviorTree::GetIndexText() const
+{
+	UBehaviorTreeGraphNode* StateNode = CastChecked<UBehaviorTreeGraphNode>(GraphNode);
+	UEdGraphPin* MyInputPin = StateNode->GetInputPin();
+	UEdGraphPin* MyParentOutputPin = NULL;
+	if (MyInputPin != NULL && MyInputPin->LinkedTo.Num() > 0)
+	{
+		MyParentOutputPin = MyInputPin->LinkedTo[0];
+	}
+
+	int32 Index = 0;
+
+	if (GEditor->bIsSimulatingInEditor || GEditor->PlayWorld != NULL)
+	{
+		// special case: range of execution indices in composite decorator node
+		UBehaviorTreeGraphNode_CompositeDecorator* CompDecorator = Cast<UBehaviorTreeGraphNode_CompositeDecorator>(GraphNode);
+		if (CompDecorator && CompDecorator->FirstExecutionIndex != CompDecorator->LastExecutionIndex)
+		{
+			return FText::Format(LOCTEXT("CompositeDecoratorFormat", "{0}..{1}"), FText::AsNumber(CompDecorator->FirstExecutionIndex), FText::AsNumber(CompDecorator->LastExecutionIndex));
+		}
+
+		// show execution index (debugging purposes)
+		UBTNode* BTNode = Cast<UBTNode>(StateNode->NodeInstance);
+		Index = (BTNode && BTNode->GetExecutionIndex() < 0xffff) ? BTNode->GetExecutionIndex() : -1;
+	}
+	else
+	{
+		// show child index
+		if (MyParentOutputPin != NULL)
+		{
+			for (Index = 0; Index < MyParentOutputPin->LinkedTo.Num(); ++Index)
+			{
+				if (MyParentOutputPin->LinkedTo[Index] == MyInputPin)
+				{
+					break;
+				}
+			}
+		}
+	}
+
+	return FText::AsNumber(Index);
+}
+
+FText SGraphNode_BehaviorTree::GetIndexTooltipText() const
+{
+	if (GEditor->bIsSimulatingInEditor || GEditor->PlayWorld != NULL)
+	{
+		return LOCTEXT("ExecutionIndexTooltip", "Execution index: this shows the order in which nodes are executed.");
+	}
+	else
+	{
+		return LOCTEXT("ChildIndexTooltip", "Child index: this shows the order in which child nodes are executed.");
+	}
+}
+
 EVisibility SGraphNode_BehaviorTree::GetBlueprintIconVisibility() const
 {
 	UBehaviorTreeGraphNode* BTGraphNode = Cast<UBehaviorTreeGraphNode>(GraphNode);
@@ -996,3 +1163,47 @@ void SGraphNode_BehaviorTree::GetOverlayBrushes(bool bSelected, const FVector2D
 		}
 	}
 }
+
+TArray<FOverlayWidgetInfo> SGraphNode_BehaviorTree::GetOverlayWidgets(bool bSelected, const FVector2D& WidgetSize) const
+{
+	TArray<FOverlayWidgetInfo> Widgets;
+
+	check(NodeBody.IsValid());
+	check(IndexOverlay.IsValid());
+
+	FVector2D Origin(0.0f, 0.0f);
+
+	// build overlays for decorator sub-nodes
+	for(const auto& DecoratorWidget : DecoratorWidgets)
+	{
+		TArray<FOverlayWidgetInfo> OverlayWidgets = DecoratorWidget->GetOverlayWidgets(bSelected, WidgetSize);
+		for(auto& OverlayWidget : OverlayWidgets)
+		{
+			OverlayWidget.OverlayOffset.Y += Origin.Y;
+		}
+		Widgets.Append(OverlayWidgets);
+		Origin.Y += DecoratorWidget->GetDesiredSize().Y;
+	}
+
+	FOverlayWidgetInfo Overlay(IndexOverlay);
+	Overlay.OverlayOffset = FVector2D(WidgetSize.X - (IndexOverlay->GetDesiredSize().X * 0.5f), Origin.Y);
+	Widgets.Add(Overlay);
+
+	Origin.Y += NodeBody->GetDesiredSize().Y;
+
+	// build overlays for service sub-nodes
+	for(const auto& ServiceWidget : ServicesWidgets)
+	{
+		TArray<FOverlayWidgetInfo> OverlayWidgets = ServiceWidget->GetOverlayWidgets(bSelected, WidgetSize);
+		for(auto& OverlayWidget : OverlayWidgets)
+		{
+			OverlayWidget.OverlayOffset.Y += Origin.Y;
+		}
+		Widgets.Append(OverlayWidgets);
+		Origin.Y += ServiceWidget->GetDesiredSize().Y;
+	}
+
+	return Widgets;
+}
+
+#undef LOCTEXT_NAMESPACE
\ No newline at end of file
diff --git a/Engine/Source/Editor/BehaviorTreeEditor/Private/SGraphNode_BehaviorTree.h b/Engine/Source/Editor/BehaviorTreeEditor/Private/SGraphNode_BehaviorTree.h
--- a/Engine/Source/Editor/BehaviorTreeEditor/Private/SGraphNode_BehaviorTree.h
+++ b/Engine/Source/Editor/BehaviorTreeEditor/Private/SGraphNode_BehaviorTree.h
@@ -20,6 +20,7 @@ public:
 	virtual FReply OnDrop( const FGeometry& MyGeometry, const FDragDropEvent& DragDropEvent ) override;
 	virtual void OnDragLeave( const FDragDropEvent& DragDropEvent ) override;
 	virtual void GetOverlayBrushes(bool bSelected, const FVector2D WidgetSize, TArray<FOverlayBrushInfo>& Brushes) const override;
+	virtual TArray<FOverlayWidgetInfo> GetOverlayWidgets(bool bSelected, const FVector2D& WidgetSize) const override;
 	// End of SGraphNode interface
 
 	/** handle double click */
@@ -93,6 +94,13 @@ protected:
 	TSharedPtr<SVerticalBox> DecoratorsBox;
 	TSharedPtr<SVerticalBox> ServicesBox;
 	TSharedPtr<SHorizontalBox> OutputPinBox;
+
+	/** The widget we use to display the index of the node */
+	TSharedPtr<SWidget> IndexOverlay;
+
+	/** The node body widget, cached here so we can determine its size when we want ot position our overlays */
+	TSharedPtr<SBorder> NodeBody;
+
 	FSlateColor GetBorderBackgroundColor() const;
 	FSlateColor GetBackgroundColor() const;
 	FString	GetDescription() const;
@@ -101,5 +109,20 @@ protected:
 	virtual const FSlateBrush* GetNameIcon() const;
 	virtual EVisibility GetBlueprintIconVisibility() const;
 
+	/** Get the visibility of the index overlay */
+	EVisibility GetIndexVisibility() const;
+
+	/** Get the text to display in the index overlay */
+	FText GetIndexText() const;
+
+	/** Get the tooltip for the index overlay */
+	FText GetIndexTooltipText() const;
+
+	/** Get the color to display for the index overlay. This changes on hover state of sibling nodes */
+	FSlateColor GetIndexColor(bool bHovered) const;
+
+	/** Handle hover state changing for the index widget - we use this to highlight sibling nodes */
+	void OnIndexHoverStateChanged(bool bHovered);
+
 	FString GetPinTooltip(UEdGraphPin* GraphPinObj) const;
 };
diff --git a/Engine/Source/Editor/EditorStyle/Private/SlateEditorStyle.cpp b/Engine/Source/Editor/EditorStyle/Private/SlateEditorStyle.cpp
--- a/Engine/Source/Editor/EditorStyle/Private/SlateEditorStyle.cpp
+++ b/Engine/Source/Editor/EditorStyle/Private/SlateEditorStyle.cpp
@@ -2909,6 +2909,15 @@ void FSlateEditorStyle::FStyle::SetupGraphEditorStyles()
 			Set( "BTEditor.Graph.BTNode.Task.RunEQSQuery.Icon", new IMAGE_BRUSH( "/BehaviorTree/Icons/EQS_24x", Icon24x24 ) );
 			Set( "BTEditor.Graph.BTNode.Task.Wait.Icon", new IMAGE_BRUSH( "/BehaviorTree/Icons/Wait_24x", Icon24x24 ) );
 			Set( "BTEditor.Graph.BTNode.Blueprint", new IMAGE_BRUSH( "/BehaviorTree/Icons/Blueprint_Referencer_16x", Icon16x16 ) );
+			Set( "BTEditor.Graph.BTNode.Index", new BOX_BRUSH( "/BehaviorTree/IndexCircle", Icon20x20, FMargin(9.0f/20.0f, 1.0f/20.0f, 9.0f/20.0f, 3.0f/20.0f) ) );
+
+			Set( "BTEditor.Graph.BTNode.Index.Color", FLinearColor(0.3f, 0.3f, 0.3f, 1.0f) );
+			Set( "BTEditor.Graph.BTNode.Index.HoveredColor", FLinearColor(1.0f, 0.0f, 0.0f, 1.0f) );
+			
+
+			FTextBlockStyle GraphNodeTitle = FTextBlockStyle(NormalText)
+				.SetFont( TTF_CORE_FONT( "Fonts/Roboto-Bold", 8 ) );
+			Set( "BTEditor.Graph.BTNode.IndexText", GraphNodeTitle );
 
 			Set( "BTEditor.Debugger.BackOver", new IMAGE_BRUSH( "Icons/icon_step_back_40x", Icon40x40 ) );
 			Set( "BTEditor.Debugger.BackOver.Small", new IMAGE_BRUSH( "Icons/icon_step_back_40x", Icon20x20 ) );
diff --git a/Engine/Source/Editor/GraphEditor/Private/SGraphPanel.cpp b/Engine/Source/Editor/GraphEditor/Private/SGraphPanel.cpp
--- a/Engine/Source/Editor/GraphEditor/Private/SGraphPanel.cpp
+++ b/Engine/Source/Editor/GraphEditor/Private/SGraphPanel.cpp
@@ -133,7 +133,7 @@ int32 SGraphPanel::OnPaint( const FGeometry& AllottedGeometry, const FSlateRect&
 	const float ZoomFactor = AllottedGeometry.Scale * GetZoomAmount();
 
 	FArrangedChildren ArrangedChildren(EVisibility::Visible);
-	this->ArrangeChildren(AllottedGeometry, ArrangedChildren);
+	ArrangeChildNodes(AllottedGeometry, ArrangedChildren);
 
 	// Determine some 'global' settings based on current LOD
 	const bool bDrawScaledCommentBubblesThisFrame = GetCurrentLOD() > EGraphRenderingLOD::LowestDetail;
@@ -282,31 +282,51 @@ int32 SGraphPanel::OnPaint( const FGeometry& AllottedGeometry, const FSlateRect&
 					// Get its size
 					const FVector2D WidgetSize = CurWidget.Geometry.Size;
 
-					TArray<FOverlayBrushInfo> OverlayBrushes;
-					ChildNode->GetOverlayBrushes(bSelected, WidgetSize, /*out*/ OverlayBrushes);
-
-					for (int32 BrushIndex = 0; BrushIndex < OverlayBrushes.Num(); ++BrushIndex)
 					{
-						FOverlayBrushInfo& OverlayInfo = OverlayBrushes[BrushIndex];
-						const FSlateBrush* OverlayBrush = OverlayInfo.Brush;
-						if(OverlayBrush != NULL)
+						TArray<FOverlayBrushInfo> OverlayBrushes;
+						ChildNode->GetOverlayBrushes(bSelected, WidgetSize, /*out*/ OverlayBrushes);
+
+						for (int32 BrushIndex = 0; BrushIndex < OverlayBrushes.Num(); ++BrushIndex)
 						{
-							FPaintGeometry BouncedGeometry = CurWidget.Geometry.ToPaintGeometry(OverlayInfo.OverlayOffset, OverlayBrush->ImageSize, 1.f);
-
-							// Handle bouncing
-							const float BounceValue = FMath::Sin(2.0f * PI * BounceCurve.GetLerpLooping());
-							BouncedGeometry.DrawPosition += (OverlayInfo.AnimationEnvelope * BounceValue * ZoomFactor);
-
-							CurWidgetsMaxLayerId++;
-							FSlateDrawElement::MakeBox(
-								OutDrawElements,
-								CurWidgetsMaxLayerId,
-								BouncedGeometry,
-								OverlayBrush,
-								MyClippingRect
-								);
+							FOverlayBrushInfo& OverlayInfo = OverlayBrushes[BrushIndex];
+							const FSlateBrush* OverlayBrush = OverlayInfo.Brush;
+							if(OverlayBrush != NULL)
+							{
+								FPaintGeometry BouncedGeometry = CurWidget.Geometry.ToPaintGeometry(OverlayInfo.OverlayOffset, OverlayBrush->ImageSize, 1.f);
+
+								// Handle bouncing
+								const float BounceValue = FMath::Sin(2.0f * PI * BounceCurve.GetLerpLooping());
+								BouncedGeometry.DrawPosition += (OverlayInfo.AnimationEnvelope * BounceValue * ZoomFactor);
+
+								CurWidgetsMaxLayerId++;
+								FSlateDrawElement::MakeBox(
+									OutDrawElements,
+									CurWidgetsMaxLayerId,
+									BouncedGeometry,
+									OverlayBrush,
+									MyClippingRect
+									);
+							}
+
 						}
+					}
 
+					{
+						TArray<FOverlayWidgetInfo> OverlayWidgets = ChildNode->GetOverlayWidgets(bSelected, WidgetSize);
+
+						for (int32 WidgetIndex = 0; WidgetIndex < OverlayWidgets.Num(); ++WidgetIndex)
+						{
+							FOverlayWidgetInfo& OverlayInfo = OverlayWidgets[WidgetIndex];
+							if(OverlayInfo.Widget->GetVisibility() == EVisibility::Visible)
+							{
+								// call SlatePrepass as these widgets are not in the 'normal' child hierarchy
+								OverlayInfo.Widget->SlatePrepass();
+
+								const FGeometry WidgetGeometry = CurWidget.Geometry.MakeChild(OverlayInfo.OverlayOffset, OverlayInfo.Widget->GetDesiredSize(), 1.f);
+
+								OverlayInfo.Widget->Paint(WidgetGeometry, MyClippingRect, OutDrawElements, CurWidgetsMaxLayerId, InWidgetStyle, bParentEnabled);
+							}
+						}
 					}
 				}
 
@@ -501,6 +521,29 @@ bool SGraphPanel::SupportsKeyboardFocus() const
 	return true;
 }
 
+void SGraphPanel::OnArrangeChildren( const FGeometry& AllottedGeometry, FArrangedChildren& ArrangedChildren ) const
+{
+	SNodePanel::OnArrangeChildren(AllottedGeometry, ArrangedChildren);
+
+	FArrangedChildren MyArrangedChildren(ArrangedChildren.GetFilter());
+	for (int32 ChildIndex = 0; ChildIndex < ArrangedChildren.Num(); ++ChildIndex)
+	{
+		FArrangedWidget& CurWidget = ArrangedChildren(ChildIndex);
+		TSharedRef<SGraphNode> ChildNode = StaticCastSharedRef<SGraphNode>(CurWidget.Widget);
+
+		TArray<FOverlayWidgetInfo> OverlayWidgets = ChildNode->GetOverlayWidgets(false, CurWidget.Geometry.Size);
+
+		for (int32 WidgetIndex = 0; WidgetIndex < OverlayWidgets.Num(); ++WidgetIndex)
+		{
+			FOverlayWidgetInfo& OverlayInfo = OverlayWidgets[WidgetIndex];
+
+			MyArrangedChildren.AddWidget(AllottedGeometry.MakeChild( OverlayInfo.Widget.ToSharedRef(), CurWidget.Geometry.Position + OverlayInfo.OverlayOffset, OverlayInfo.Widget->GetDesiredSize(), GetZoomAmount() ));
+		}
+	}
+
+	ArrangedChildren.Append(MyArrangedChildren);
+}
+
 void SGraphPanel::UpdateSelectedNodesPositions (FVector2D PositionIncrement)
 {
 	for (FGraphPanelSelectionSet::TIterator NodeIt(SelectionManager.SelectedNodes); NodeIt; ++NodeIt)
diff --git a/Engine/Source/Editor/GraphEditor/Private/SGraphPanel.h b/Engine/Source/Editor/GraphEditor/Private/SGraphPanel.h
--- a/Engine/Source/Editor/GraphEditor/Private/SGraphPanel.h
+++ b/Engine/Source/Editor/GraphEditor/Private/SGraphPanel.h
@@ -100,6 +100,7 @@ public:
 	virtual int32 OnPaint( const FGeometry& AllottedGeometry, const FSlateRect& MyClippingRect, FSlateWindowElementList& OutDrawElements, int32 LayerId, const FWidgetStyle& InWidgetStyle, bool bParentEnabled ) const override;
 	virtual FReply OnKeyDown( const FGeometry& MyGeometry, const FKeyboardEvent& InKeyboardEvent ) override;
 	virtual bool SupportsKeyboardFocus() const override;
+	virtual void OnArrangeChildren( const FGeometry& AllottedGeometry, FArrangedChildren& ArrangedChildren ) const override;
 	// End of SWidget interface
 
 	// SNodePanel interface
diff --git a/Engine/Source/Editor/GraphEditor/Private/SNodePanel.cpp b/Engine/Source/Editor/GraphEditor/Private/SNodePanel.cpp
--- a/Engine/Source/Editor/GraphEditor/Private/SNodePanel.cpp
+++ b/Engine/Source/Editor/GraphEditor/Private/SNodePanel.cpp
@@ -199,6 +199,11 @@ namespace NodePanelDefs
 };
 
 void SNodePanel::OnArrangeChildren(const FGeometry& AllottedGeometry, FArrangedChildren& ArrangedChildren) const
+{
+	ArrangeChildNodes(AllottedGeometry, ArrangedChildren);
+}
+
+void SNodePanel::ArrangeChildNodes(const FGeometry& AllottedGeometry, FArrangedChildren& ArrangedChildren) const
 {
 	const TSlotlessChildren<SNode>& ChildrenToArrange = ArrangedChildren.Accepts(EVisibility::Hidden) ? Children : VisibleChildren;
 	// First pass nodes
@@ -555,7 +560,7 @@ FReply SNodePanel::OnMouseButtonDown( const FGeometry& MyGeometry, const FPointe
 	{
 		// LEFT BUTTON is for selecting nodes and manipulating pins.
 		FArrangedChildren ArrangedChildren(EVisibility::Visible);
-		ArrangeChildren(MyGeometry, ArrangedChildren);
+		ArrangeChildNodes(MyGeometry, ArrangedChildren);
 
 		const int32 NodeUnderMouseIndex = SWidget::FindChildUnderMouse( ArrangedChildren, MouseEvent );
 		if ( NodeUnderMouseIndex != INDEX_NONE )
diff --git a/Engine/Source/Editor/GraphEditor/Public/SNodePanel.h b/Engine/Source/Editor/GraphEditor/Public/SNodePanel.h
--- a/Engine/Source/Editor/GraphEditor/Public/SNodePanel.h
+++ b/Engine/Source/Editor/GraphEditor/Public/SNodePanel.h
@@ -173,6 +173,30 @@ public:
 	}
 };
 
+// Entry for an overlay widget in the node panel
+struct FOverlayWidgetInfo
+{
+public:
+	/** Widget to use */
+	TSharedPtr<SWidget> Widget;
+
+	/** Offset origin of the overlay from the widget */
+	FVector2D OverlayOffset;
+
+public:
+	FOverlayWidgetInfo()
+		: Widget(nullptr)
+		, OverlayOffset(0.f, 0.f)
+	{
+	}
+
+	FOverlayWidgetInfo(TSharedPtr<SWidget> InWidget)
+		: Widget(InWidget)
+		, OverlayOffset(0.f, 0.f)
+	{
+	}
+};
+
 // Entry for an information popup in the node panel
 struct FGraphInformationPopupInfo
 {
@@ -292,10 +316,16 @@ public:
 		{
 		}
 
+		/** Populate the widgets array with any overlay widgets to render */
+		virtual TArray<FOverlayWidgetInfo> GetOverlayWidgets(bool bSelected, const FVector2D& WidgetSize) const
+		{
+			return TArray<FOverlayWidgetInfo>();
+		}
+
 		/** Populate the popups array with any popups to render */
 		virtual void GetNodeInfoPopups(FNodeInfoContext* Context, TArray<FGraphInformationPopupInfo>& Popups) const
 		{
-		}
+		}	
 
 		/** Returns true if this node is dependent on the location of other nodes (it can only depend on the location of first-pass only nodes) */
 		virtual bool RequiresSecondPassLayout() const
@@ -477,6 +507,9 @@ protected:
 	/** Populate visibile children array */
 	virtual void PopulateVisibleChildren(const FGeometry& AllottedGeometry);
 
+	/** Arrange child nodes - allows derived classes to supply non-node children in OnArrangeChildren */
+	virtual void ArrangeChildNodes(const FGeometry& AllottedGeometry, FArrangedChildren& ArrangedChildren) const;
+
 	// Paint the background as lines
 	void PaintBackgroundAsLines(const FSlateBrush* BackgroundImage, const FGeometry& AllottedGeometry, const FSlateRect& MyClippingRect, FSlateWindowElementList& OutDrawElements, int32& DrawLayerId) const;
 
