Behaviour Trees are now correctly focussed on double-click

Correctly called through to parent implementation of FocusWindow().
Also prevent blackboards re-used by multiple BTs causing a mode switch when it is not needed.

TTP# 341365 - CRITICAL: Regression: EDITOR: Double clicking a Behavior Tree asset does not switch to its tab when open

diff --git a/Engine/Source/Editor/BehaviorTreeEditor/Private/AssetTypeActions_BehaviorTree.cpp b/Engine/Source/Editor/BehaviorTreeEditor/Private/AssetTypeActions_BehaviorTree.cpp
--- a/Engine/Source/Editor/BehaviorTreeEditor/Private/AssetTypeActions_BehaviorTree.cpp
+++ b/Engine/Source/Editor/BehaviorTreeEditor/Private/AssetTypeActions_BehaviorTree.cpp
@@ -69,7 +69,7 @@ void FAssetTypeActions_BehaviorTree::OpenAssetEditor( const TArray<UObject*>& In
 			bool bFoundExisting = false;
 			if(BehaviorTree->BlackboardAsset != nullptr)
 			{
-				const bool bFocusIfOpen = true;
+				const bool bFocusIfOpen = false;
 				FBehaviorTreeEditor* ExistingInstance = static_cast<FBehaviorTreeEditor*>(FAssetEditorManager::Get().FindEditorForAsset(BehaviorTree->BlackboardAsset, bFocusIfOpen));
 				if(ExistingInstance != nullptr && ExistingInstance->GetBehaviorTree() == nullptr)
 				{
diff --git a/Engine/Source/Editor/BehaviorTreeEditor/Private/BehaviorTreeEditor.cpp b/Engine/Source/Editor/BehaviorTreeEditor/Private/BehaviorTreeEditor.cpp
--- a/Engine/Source/Editor/BehaviorTreeEditor/Private/BehaviorTreeEditor.cpp
+++ b/Engine/Source/Editor/BehaviorTreeEditor/Private/BehaviorTreeEditor.cpp
@@ -1943,6 +1943,8 @@ void FBehaviorTreeEditor::FocusWindow(UObject* ObjectToFocusOn)
 	{
 		SetCurrentMode(BlackboardMode);
 	}
+
+	FWorkflowCentricApplication::FocusWindow(ObjectToFocusOn);
 }
 
 void FBehaviorTreeEditor::InitializeDebuggerState(class FBehaviorTreeDebugger* ParentDebugger) const
