Behavior Trees are no longer 'experimental'

Removed settings that locked use of BTs and BT editor features.
Module is now always loaded in editor.

diff --git a/Engine/Source/Editor/BehaviorTreeEditor/Private/AssetTypeActions_BehaviorTree.cpp b/Engine/Source/Editor/BehaviorTreeEditor/Private/AssetTypeActions_BehaviorTree.cpp
--- a/Engine/Source/Editor/BehaviorTreeEditor/Private/AssetTypeActions_BehaviorTree.cpp
+++ b/Engine/Source/Editor/BehaviorTreeEditor/Private/AssetTypeActions_BehaviorTree.cpp
@@ -13,31 +13,13 @@
 
 #define LOCTEXT_NAMESPACE "AssetTypeActions"
 
-bool FAssetTypeActions_BehaviorTree::BehaviorTreeEditorEnabled()
-{
-	bool BehaviorTreeNewAssetsEnabled=false;
-	GConfig->GetBool(TEXT("BehaviorTreesEd"), TEXT("BehaviorTreeNewAssetsEnabled"), BehaviorTreeNewAssetsEnabled, GEngineIni);
-
-	if (!BehaviorTreeNewAssetsEnabled && !GetDefault<UEditorExperimentalSettings>()->bBehaviorTreeEditor)
-	{
-		return false;
-	}
-
-	return true;
-}
-
 bool FAssetTypeActions_BehaviorTree::HasActions ( const TArray<UObject*>& InObjects ) const
 {
-	return BehaviorTreeEditorEnabled();
+	return true;
 }
 
 void FAssetTypeActions_BehaviorTree::GetActions( const TArray<UObject*>& InObjects, FMenuBuilder& MenuBuilder )
 {
-	if(!BehaviorTreeEditorEnabled())
-	{
-		return;
-	}
-
 	auto BehaviorTrees = GetTypedWeakObjectPtrs<UBehaviorTree>(InObjects);
 
 	MenuBuilder.AddMenuEntry(
@@ -53,11 +35,6 @@ void FAssetTypeActions_BehaviorTree::GetActions( const TArray<UObject*>& InObjec
 
 void FAssetTypeActions_BehaviorTree::OpenAssetEditor( const TArray<UObject*>& InObjects, TSharedPtr<IToolkitHost> EditWithinLevelEditor )
 {
-	if(!BehaviorTreeEditorEnabled())
-	{
-		return;
-	}
-
 	EToolkitMode::Type Mode = EditWithinLevelEditor.IsValid() ? EToolkitMode::WorldCentric : EToolkitMode::Standalone;
 
 	for(auto Object : InObjects)
diff --git a/Engine/Source/Editor/BehaviorTreeEditor/Private/AssetTypeActions_BehaviorTree.h b/Engine/Source/Editor/BehaviorTreeEditor/Private/AssetTypeActions_BehaviorTree.h
--- a/Engine/Source/Editor/BehaviorTreeEditor/Private/AssetTypeActions_BehaviorTree.h
+++ b/Engine/Source/Editor/BehaviorTreeEditor/Private/AssetTypeActions_BehaviorTree.h
@@ -16,9 +16,6 @@ public:
 	virtual uint32 GetCategories() override { return EAssetTypeCategories::Misc; }
 	virtual void PerformAssetDiff(UObject* Asset1, UObject* Asset2, const struct FRevisionInfo& OldRevision, const struct FRevisionInfo& NewRevision) const override;
 
-	/** Check whether the behavior tree system is enabled */
-	static bool BehaviorTreeEditorEnabled();
-
 private:
 
 	/* Called to open the Behavior Tree defaults view, this opens whatever text dif tool the user has */
diff --git a/Engine/Source/Editor/BehaviorTreeEditor/Private/AssetTypeActions_Blackboard.cpp b/Engine/Source/Editor/BehaviorTreeEditor/Private/AssetTypeActions_Blackboard.cpp
--- a/Engine/Source/Editor/BehaviorTreeEditor/Private/AssetTypeActions_Blackboard.cpp
+++ b/Engine/Source/Editor/BehaviorTreeEditor/Private/AssetTypeActions_Blackboard.cpp
@@ -19,16 +19,11 @@ UClass* FAssetTypeActions_Blackboard::GetSupportedClass() const
 
 bool FAssetTypeActions_Blackboard::HasActions ( const TArray<UObject*>& InObjects ) const
 {
-	return FAssetTypeActions_BehaviorTree::BehaviorTreeEditorEnabled();
+	return true;
 }
 
 void FAssetTypeActions_Blackboard::GetActions( const TArray<UObject*>& InObjects, FMenuBuilder& MenuBuilder )
 {
-	if (!FAssetTypeActions_BehaviorTree::BehaviorTreeEditorEnabled())
-	{
-		return;
-	}
-
 	auto BlackboardData = GetTypedWeakObjectPtrs<UBlackboardData>(InObjects);
 
 	MenuBuilder.AddMenuEntry(
@@ -44,11 +39,6 @@ void FAssetTypeActions_Blackboard::GetActions( const TArray<UObject*>& InObjects
 
 void FAssetTypeActions_Blackboard::OpenAssetEditor( const TArray<UObject*>& InObjects, TSharedPtr<class IToolkitHost> EditWithinLevelEditor )
 {
-	if (!FAssetTypeActions_BehaviorTree::BehaviorTreeEditorEnabled())
-	{
-		return;
-	}
-
 	EToolkitMode::Type Mode = EditWithinLevelEditor.IsValid() ? EToolkitMode::WorldCentric : EToolkitMode::Standalone;
 
 	for(auto Object : InObjects)
diff --git a/Engine/Source/Editor/BehaviorTreeEditor/Private/BehaviorTreeFactory.cpp b/Engine/Source/Editor/BehaviorTreeEditor/Private/BehaviorTreeFactory.cpp
--- a/Engine/Source/Editor/BehaviorTreeEditor/Private/BehaviorTreeFactory.cpp
+++ b/Engine/Source/Editor/BehaviorTreeEditor/Private/BehaviorTreeFactory.cpp
@@ -16,9 +16,7 @@ UBehaviorTreeFactory::UBehaviorTreeFactory(const class FPostConstructInitializeP
 
 bool UBehaviorTreeFactory::CanCreateNew() const
 {
-	bool bBehaviorTreeNewAssetsEnabled = false;
-	GConfig->GetBool(TEXT("BehaviorTreesEd"), TEXT("BehaviorTreeNewAssetsEnabled"), bBehaviorTreeNewAssetsEnabled, GEngineIni);
-	return (bBehaviorTreeNewAssetsEnabled || GetDefault<UEditorExperimentalSettings>()->bBehaviorTreeEditor);
+	return true;
 }
 
 UObject* UBehaviorTreeFactory::FactoryCreateNew(UClass* Class,UObject* InParent,FName Name,EObjectFlags Flags,UObject* Context,FFeedbackContext* Warn)
diff --git a/Engine/Source/Editor/BehaviorTreeEditor/Private/BehaviorTreeGraphNode.cpp b/Engine/Source/Editor/BehaviorTreeEditor/Private/BehaviorTreeGraphNode.cpp
--- a/Engine/Source/Editor/BehaviorTreeEditor/Private/BehaviorTreeGraphNode.cpp
+++ b/Engine/Source/Editor/BehaviorTreeEditor/Private/BehaviorTreeGraphNode.cpp
@@ -20,6 +20,7 @@ UBehaviorTreeGraphNode::UBehaviorTreeGraphNode(const class FPostConstructInitial
 	bHighlightInSearchRange0 = false;
 	bHighlightInSearchRange1 = false;
 	bHighlightInSearchTree = false;
+	bHighlightChildNodeIndices = false;
 	bRootLevel = false;
 	bInjectedNode = false;
 	bHasObserverError = false;
diff --git a/Engine/Source/Editor/BehaviorTreeEditor/Private/BlackboardDataFactory.cpp b/Engine/Source/Editor/BehaviorTreeEditor/Private/BlackboardDataFactory.cpp
--- a/Engine/Source/Editor/BehaviorTreeEditor/Private/BlackboardDataFactory.cpp
+++ b/Engine/Source/Editor/BehaviorTreeEditor/Private/BlackboardDataFactory.cpp
@@ -16,9 +16,7 @@ UBlackboardDataFactory::UBlackboardDataFactory(const class FPostConstructInitial
 
 bool UBlackboardDataFactory::CanCreateNew() const
 {
-	bool bBehaviorTreeNewAssetsEnabled = false;
-	GConfig->GetBool(TEXT("BehaviorTreesEd"), TEXT("BehaviorTreeNewAssetsEnabled"), bBehaviorTreeNewAssetsEnabled, GEngineIni);
-	return (bBehaviorTreeNewAssetsEnabled || GetDefault<UEditorExperimentalSettings>()->bBehaviorTreeEditor);
+	return true;
 }
 
 UObject* UBlackboardDataFactory::FactoryCreateNew(UClass* Class,UObject* InParent,FName Name,EObjectFlags Flags,UObject* Context,FFeedbackContext* Warn)
diff --git a/Engine/Source/Editor/UnrealEd/Classes/Settings/EditorExperimentalSettings.h b/Engine/Source/Editor/UnrealEd/Classes/Settings/EditorExperimentalSettings.h
--- a/Engine/Source/Editor/UnrealEd/Classes/Settings/EditorExperimentalSettings.h
+++ b/Engine/Source/Editor/UnrealEd/Classes/Settings/EditorExperimentalSettings.h
@@ -24,10 +24,6 @@ public:
 	UPROPERTY(EditAnywhere, config, Category = Tools, meta = ( DisplayName = "Unreal Motion Graphics (UMG)" ))
 	bool bUnrealMotionGraphics;
 
-	/** Allows to open editor for Behavior Tree assets */
-	UPROPERTY(EditAnywhere, config, Category=Tools, meta=(DisplayName="Behavior Tree Editor"))
-	bool bBehaviorTreeEditor;
-
 	/** Allows usage of the Translation Editor */
 	UPROPERTY(EditAnywhere, config, Category = Tools, meta = (DisplayName = "Translation Editor"))
 	bool bEnableTranslationEditor;
diff --git a/Engine/Source/Editor/UnrealEd/Private/Editor.cpp b/Engine/Source/Editor/UnrealEd/Private/Editor.cpp
--- a/Engine/Source/Editor/UnrealEd/Private/Editor.cpp
+++ b/Engine/Source/Editor/UnrealEd/Private/Editor.cpp
@@ -555,6 +555,7 @@ void UEditorEngine::Init(IEngineLoop* InEngineLoop)
 		FModuleManager::Get().LoadModule(TEXT("UndoHistory"));
 		FModuleManager::Get().LoadModule(TEXT("DeviceProfileEditor"));
 		FModuleManager::Get().LoadModule(TEXT("SourceCodeAccess"));
+		FModuleManager::Get().LoadModule(TEXT("BehaviorTreeEditor"));
 
 		if (!IsRunningCommandlet())
 		{
@@ -566,14 +567,6 @@ void UEditorEngine::Init(IEngineLoop* InEngineLoop)
 		{
 			FModuleManager::Get().LoadModule(TEXT("PListEditor"));
 		}
-		
-		//check if we need to load behavior tree editor module (it could be loaded earlier) 
-		bool bBehaviorTreeEditorEnabled = false;
-		GConfig->GetBool(TEXT("BehaviorTreesEd"), TEXT("BehaviorTreeEditorEnabled"), bBehaviorTreeEditorEnabled, GEngineIni);
-		if ( (GetDefault<UEditorExperimentalSettings>()->bBehaviorTreeEditor || bBehaviorTreeEditorEnabled) && UBehaviorTreeManager::IsBehaviorTreeUsageEnabled() && !FModuleManager::Get().IsModuleLoaded(TEXT("BehaviorTreeEditor")))
-		{
-			FModuleManager::Get().LoadModule(TEXT("BehaviorTreeEditor"));
-		}
 
 		bool bEnvironmentQueryEditor = false;
 		GConfig->GetBool(TEXT("EnvironmentQueryEd"), TEXT("EnableEnvironmentQueryEd"), bEnvironmentQueryEditor, GEngineIni);
diff --git a/Engine/Source/Editor/UnrealEd/Private/Settings/SettingsClasses.cpp b/Engine/Source/Editor/UnrealEd/Private/Settings/SettingsClasses.cpp
--- a/Engine/Source/Editor/UnrealEd/Private/Settings/SettingsClasses.cpp
+++ b/Engine/Source/Editor/UnrealEd/Private/Settings/SettingsClasses.cpp
@@ -63,11 +63,7 @@ void UEditorExperimentalSettings::PostEditChangeProperty( struct FPropertyChange
 
 	const FName Name = (PropertyChangedEvent.Property != nullptr) ? PropertyChangedEvent.Property->GetFName() : NAME_None;
 
-	if (Name == FName(TEXT("bBehaviorTreeEditor")) && bBehaviorTreeEditor && !FModuleManager::Get().IsModuleLoaded(TEXT("BehaviorTreeEditor")))
-	{
-		FModuleManager::Get().LoadModule(TEXT("BehaviorTreeEditor"));
-	}
-	else if (Name == FName(TEXT("ConsoleForGamepadLabels")))
+	if (Name == FName(TEXT("ConsoleForGamepadLabels")))
 	{
 		EKeys::SetConsoleForGamepadLabels(ConsoleForGamepadLabels);
 	}
diff --git a/Engine/Source/Runtime/AIModule/Classes/BehaviorTree/BehaviorTreeManager.h b/Engine/Source/Runtime/AIModule/Classes/BehaviorTree/BehaviorTreeManager.h
--- a/Engine/Source/Runtime/AIModule/Classes/BehaviorTree/BehaviorTreeManager.h
+++ b/Engine/Source/Runtime/AIModule/Classes/BehaviorTree/BehaviorTreeManager.h
@@ -32,9 +32,6 @@ class AIMODULE_API UBehaviorTreeManager : public UObject
 	/** get behavior tree template for given blueprint */
 	bool LoadTree(class UBehaviorTree* Asset, class UBTCompositeNode*& Root, uint16& InstanceMemorySize);
 
-	/** check if BT system is to be used. Here only for BT development time */
-	static bool IsBehaviorTreeUsageEnabled();
-
 	/** get aligned memory size */
 	static int32 GetAlignedDataSize(int32 Size);
 
diff --git a/Engine/Source/Runtime/AIModule/Private/BehaviorTree/BehaviorTreeComponent.cpp b/Engine/Source/Runtime/AIModule/Private/BehaviorTree/BehaviorTreeComponent.cpp
--- a/Engine/Source/Runtime/AIModule/Private/BehaviorTree/BehaviorTreeComponent.cpp
+++ b/Engine/Source/Runtime/AIModule/Private/BehaviorTree/BehaviorTreeComponent.cpp
@@ -95,12 +95,6 @@ bool UBehaviorTreeComponent::IsPaused() const
 
 bool UBehaviorTreeComponent::StartTree(class UBehaviorTree* TreeAsset, EBTExecutionMode::Type ExecuteMode)
 {
-	if (UBehaviorTreeManager::IsBehaviorTreeUsageEnabled() == false)
-	{
-		UE_VLOG(GetOwner(), LogBehaviorTree, Error, TEXT("Behavior tree usage is not enabled"));
-		return false;
-	}
-
 	// clear instance stack, start should always run new tree from root
 	class UBehaviorTree* CurrentRoot = GetRootTree();
 	
diff --git a/Engine/Source/Runtime/AIModule/Private/BehaviorTree/BehaviorTreeManager.cpp b/Engine/Source/Runtime/AIModule/Private/BehaviorTree/BehaviorTreeManager.cpp
--- a/Engine/Source/Runtime/AIModule/Private/BehaviorTree/BehaviorTreeManager.cpp
+++ b/Engine/Source/Runtime/AIModule/Private/BehaviorTree/BehaviorTreeManager.cpp
@@ -29,33 +29,6 @@ void UBehaviorTreeManager::FinishDestroy()
 	Super::FinishDestroy();
 }
 
-bool UBehaviorTreeManager::IsBehaviorTreeUsageEnabled()
-{
-	struct FCheckMeStaticly
-	{
-		bool bBehaviorTreeEditorEnabled;
-		bool bBehaviorTreeUsageEnabled;
-		FCheckMeStaticly() : bBehaviorTreeEditorEnabled(false), bBehaviorTreeUsageEnabled(true)
-		{
-#if WITH_EDITOR
-			bool bBehaviorTreeEditorEnabledFromUserSettings = false;
-			GConfig->GetBool(TEXT("/Script/UnrealEd.EditorExperimentalSettings"), TEXT("bBehaviorTreeEditor"), bBehaviorTreeEditorEnabledFromUserSettings, GEditorUserSettingsIni);
-
-			GConfig->GetBool(TEXT("BehaviorTreesEd"), TEXT("BehaviorTreeEditorEnabled"), bBehaviorTreeEditorEnabled, GEngineIni);
-			bBehaviorTreeEditorEnabled = bBehaviorTreeEditorEnabled || bBehaviorTreeEditorEnabledFromUserSettings;
-			if (bBehaviorTreeUsageEnabled && bBehaviorTreeEditorEnabled && !FModuleManager::Get().IsModuleLoaded(TEXT("BehaviorTreeEditor")))
-			{
-				// load this module earlier, before any access to BehaviorTree assets
-				FModuleManager::Get().LoadModule(TEXT("BehaviorTreeEditor"));
-			}
-#endif
-		}
-	};
-
-	static FCheckMeStaticly CheckMeStaticly;
-	return CheckMeStaticly.bBehaviorTreeUsageEnabled;
-}
-
 int32 UBehaviorTreeManager::GetAlignedDataSize(int32 Size)
 {
 	// round to 4 bytes
