Fixed crash on shutdown when editing behavior trees.

diff --git a/Engine/Source/Editor/BehaviorTreeEditor/Private/BehaviorTreeEditorTypes.cpp b/Engine/Source/Editor/BehaviorTreeEditor/Private/BehaviorTreeEditorTypes.cpp
--- a/Engine/Source/Editor/BehaviorTreeEditor/Private/BehaviorTreeEditorTypes.cpp
+++ b/Engine/Source/Editor/BehaviorTreeEditor/Private/BehaviorTreeEditorTypes.cpp
@@ -166,9 +166,12 @@ FClassBrowseHelper::~FClassBrowseHelper()
 		AssetRegistryModule.Get().OnAssetRemoved().RemoveAll(this);
 
 		// Unregister to have Populate called when doing a Hot Reload.
-		IHotReloadInterface& HotReloadSupport = FModuleManager::LoadModuleChecked<IHotReloadInterface>("HotReload");
-		HotReloadSupport.OnHotReload().RemoveAll(this);
-
+		if (FModuleManager::Get().IsModuleLoaded(TEXT("HotReload")))
+		{
+			IHotReloadInterface& HotReloadSupport = FModuleManager::GetModuleChecked<IHotReloadInterface>("HotReload");
+			HotReloadSupport.OnHotReload().RemoveAll(this);
+		}
+	
 		// Unregister to have Populate called when a Blueprint is compiled.
 		GEditor->OnBlueprintCompiled().RemoveAll(this);
 		GEditor->OnClassPackageLoadedOrUnloaded().RemoveAll(this);
