fixed multiple nodes connected to special pins in behavior tree editor

diff --git a/Engine/Source/Editor/BehaviorTreeEditor/Private/EdGraphSchema_BehaviorTree.cpp b/Engine/Source/Editor/BehaviorTreeEditor/Private/EdGraphSchema_BehaviorTree.cpp
--- a/Engine/Source/Editor/BehaviorTreeEditor/Private/EdGraphSchema_BehaviorTree.cpp
+++ b/Engine/Source/Editor/BehaviorTreeEditor/Private/EdGraphSchema_BehaviorTree.cpp
@@ -590,7 +590,15 @@ const FPinConnectionResponse UEdGraphSchema_BehaviorTree::CanCreateConnection(co
 		return FPinConnectionResponse(CONNECT_RESPONSE_DISALLOW, LOCTEXT("PinErrorcycle", "Can't create a graph cycle"));
 	}
 
-	if(PinB->LinkedTo.Num() > 0)
+	const bool bPinASingleLink = bPinAIsSingleComposite || bPinAIsSingleTask || bPinAIsSingleNode;
+	const bool bPinBSingleLink = bPinBIsSingleComposite || bPinBIsSingleTask || bPinBIsSingleNode;
+	if ((bPinASingleLink && PinA->LinkedTo.Num() > 0) ||
+		(bPinBSingleLink && PinB->LinkedTo.Num() > 0))
+	{
+		return FPinConnectionResponse(CONNECT_RESPONSE_DISALLOW, LOCTEXT("PinErrorSingleConnection", "Can't connect to multiple nodes"));
+	}
+
+	if (PinB->LinkedTo.Num() > 0)
 	{
 		return FPinConnectionResponse(CONNECT_RESPONSE_BREAK_OTHERS_B, LOCTEXT("PinConnectReplace", "Replace connection"));
 	}
