Duplicate-named Blackboard entries can no longer be created.

Added checks for duplicate names when renaming entries in the Blackboard editor.
Added check when creating new keys so that we dont create duplicates at first.
Disabled name editing in the details panel, as we shoud only be editing the names in the list.

TTP# 343597 - BlackBoard editor allows multiple keys of the same name [modified]

diff --git a/Engine/Source/Editor/BehaviorTreeEditor/Private/DetailCustomizations/BlackboardDataDetails.cpp b/Engine/Source/Editor/BehaviorTreeEditor/Private/DetailCustomizations/BlackboardDataDetails.cpp
--- a/Engine/Source/Editor/BehaviorTreeEditor/Private/DetailCustomizations/BlackboardDataDetails.cpp
+++ b/Engine/Source/Editor/BehaviorTreeEditor/Private/DetailCustomizations/BlackboardDataDetails.cpp
@@ -34,7 +34,24 @@ void FBlackboardDataDetails::CustomizeDetails( IDetailLayoutBuilder& DetailLayou
 			TSharedPtr<IPropertyHandle> KeyHandle = KeysHandle->GetChildHandle((uint32)CurrentSelection);
 
 			IDetailCategoryBuilder& DetailCategoryBuilder = DetailLayout.EditCategory("Key");
-			DetailCategoryBuilder.AddProperty(KeyHandle);
+			TSharedPtr<IPropertyHandle> EntryNameProperty = KeyHandle->GetChildHandle(GET_MEMBER_NAME_CHECKED(FBlackboardEntry, EntryName));
+			DetailCategoryBuilder.AddCustomRow(TEXT("Entry Name"))
+			.NameContent()
+			[
+				EntryNameProperty->CreatePropertyNameWidget()
+			]
+			.ValueContent()
+			[
+				SNew(SHorizontalBox)
+				.IsEnabled(false)
+				+SHorizontalBox::Slot()
+				[
+					EntryNameProperty->CreatePropertyValueWidget()
+				]
+			];
+
+			TSharedPtr<IPropertyHandle> KeyTypeProperty = KeyHandle->GetChildHandle(GET_MEMBER_NAME_CHECKED(FBlackboardEntry, KeyType));
+			DetailCategoryBuilder.AddProperty(KeyTypeProperty);
 		}	
 	}
 }
diff --git a/Engine/Source/Editor/BehaviorTreeEditor/Private/SBehaviorTreeBlackboardEditor.cpp b/Engine/Source/Editor/BehaviorTreeEditor/Private/SBehaviorTreeBlackboardEditor.cpp
--- a/Engine/Source/Editor/BehaviorTreeEditor/Private/SBehaviorTreeBlackboardEditor.cpp
+++ b/Engine/Source/Editor/BehaviorTreeEditor/Private/SBehaviorTreeBlackboardEditor.cpp
@@ -174,6 +174,32 @@ void SBehaviorTreeBlackboardEditor::HandleKeyClassPicked(UClass* InClass)
 	NewKeyName = NewKeyName.Replace(TEXT(" "), TEXT(""));
 	NewKeyName += TEXT("Key");
 
+	int32 IndexSuffix = -1;
+	auto DuplicateFunction = [&](const FBlackboardEntry& Key)
+	{		
+		if(Key.EntryName.ToString() == NewKeyName)
+		{
+			IndexSuffix = FMath::Max(0, IndexSuffix);
+		}
+		if(Key.EntryName.ToString().StartsWith(NewKeyName))
+		{
+			const FString ExistingSuffix = Key.EntryName.ToString().RightChop(NewKeyName.Len());
+			if(ExistingSuffix.IsNumeric())
+			{
+				IndexSuffix = FMath::Max(FCString::Atoi(*ExistingSuffix) + 1, IndexSuffix);
+			}
+		}
+	};
+
+	// check for existing keys of the same name
+	for(const auto& Key : BlackboardData->Keys) { DuplicateFunction(Key); };
+	for(const auto& Key : BlackboardData->ParentKeys) { DuplicateFunction(Key); };
+
+	if(IndexSuffix != -1)
+	{
+		NewKeyName += FString::Printf(TEXT("%d"), IndexSuffix);
+	}
+
 	FBlackboardEntry Entry;
 	Entry.EntryName = FName(*NewKeyName);
 	Entry.KeyType = ConstructObject<UBlackboardKeyType>(InClass, BlackboardData);		
diff --git a/Engine/Source/Editor/BehaviorTreeEditor/Private/SBehaviorTreeBlackboardView.cpp b/Engine/Source/Editor/BehaviorTreeEditor/Private/SBehaviorTreeBlackboardView.cpp
--- a/Engine/Source/Editor/BehaviorTreeEditor/Private/SBehaviorTreeBlackboardView.cpp
+++ b/Engine/Source/Editor/BehaviorTreeEditor/Private/SBehaviorTreeBlackboardView.cpp
@@ -187,6 +187,35 @@ private:
 		}
 	}
 
+	virtual bool OnNameTextVerifyChanged(const FText& InNewText, FText& OutErrorMessage) override
+	{
+		check(ActionPtr.Pin()->GetTypeId() == FEdGraphSchemaAction_BlackboardEntry::StaticGetTypeId());
+		TSharedPtr<FEdGraphSchemaAction_BlackboardEntry> BlackboardEntryAction = StaticCastSharedPtr<FEdGraphSchemaAction_BlackboardEntry>(ActionPtr.Pin());
+
+		const FString NewTextAsString = InNewText.ToString();
+
+		// check for duplicate keys
+		for(const auto& Key : BlackboardEntryAction->BlackboardData->Keys)
+		{
+			if(&BlackboardEntryAction->Key != &Key && Key.EntryName.ToString() == NewTextAsString)
+			{
+				OutErrorMessage = LOCTEXT("DuplicateKeyWarning", "A key of this name already exists.");
+				return false;
+			}
+		}
+
+		for(const auto& Key : BlackboardEntryAction->BlackboardData->ParentKeys)
+		{
+			if(&BlackboardEntryAction->Key != &Key && Key.EntryName.ToString() == NewTextAsString)
+			{
+				OutErrorMessage = LOCTEXT("DuplicateParentKeyWarning", "An inherited key of this name already exists.");
+				return false;
+			}
+		}
+
+		return true;
+	}
+
 	/** Create widget for displaying debug information about this blackboard entry */
 	TSharedRef<SWidget> CreateDebugSlotWidget(const FSlateFontInfo& InFontInfo)
 	{
