Fixed Behavior Tree details was not scrolling when the details panel was larger than the tab

Reported on UDN - https://udn.unrealengine.com/questions/218866/behaviour-tree-details-panel-scrollbar-missing.html

diff --git a/Engine/Source/Editor/BehaviorTreeEditor/Private/BehaviorTreeEditor.cpp b/Engine/Source/Editor/BehaviorTreeEditor/Private/BehaviorTreeEditor.cpp
--- a/Engine/Source/Editor/BehaviorTreeEditor/Private/BehaviorTreeEditor.cpp
+++ b/Engine/Source/Editor/BehaviorTreeEditor/Private/BehaviorTreeEditor.cpp
@@ -612,9 +612,8 @@ TSharedRef<SWidget> FBehaviorTreeEditor::SpawnProperties()
 	return
 		SNew(SVerticalBox)
 		+SVerticalBox::Slot()
-		.AutoHeight()
+		.FillHeight(1.0f)
 		.HAlign(HAlign_Fill)
-		.VAlign(VAlign_Top)
 		[
 			DetailsView.ToSharedRef()
 		]
