Fixed crash selecting enum key value when enum was invalid

Reported on AnswerHub here: https://answers.unrealengine.com/questions/90206/editor-crashes-when-using-undefined-enum-in-behavi.html

diff --git a/Engine/Source/Editor/BehaviorTreeEditor/Private/DetailCustomizations/BlackboardDecoratorDetails.cpp b/Engine/Source/Editor/BehaviorTreeEditor/Private/DetailCustomizations/BlackboardDecoratorDetails.cpp
--- a/Engine/Source/Editor/BehaviorTreeEditor/Private/DetailCustomizations/BlackboardDecoratorDetails.cpp
+++ b/Engine/Source/Editor/BehaviorTreeEditor/Private/DetailCustomizations/BlackboardDecoratorDetails.cpp
@@ -185,7 +185,7 @@ TSharedRef<SWidget> FBlackboardDecoratorDetails::OnGetEnumValueContent() const
 
 FString FBlackboardDecoratorDetails::GetCurrentEnumValueDesc() const
 {
-	if (CachedCustomObjectType)
+	if (CachedCustomObjectType && EnumPropValues.Num() > 0)
 	{
 		int32 CurrentIntValue = 0;
 		FPropertyAccess::Result Result = IntValueProperty->GetValue(CurrentIntValue);
