Fix non-unity error in blackboarddatadetailes.cpp

diff --git a/Engine/Source/Editor/BehaviorTreeEditor/Private/DetailCustomizations/BlackboardDataDetails.cpp b/Engine/Source/Editor/BehaviorTreeEditor/Private/DetailCustomizations/BlackboardDataDetails.cpp
--- a/Engine/Source/Editor/BehaviorTreeEditor/Private/DetailCustomizations/BlackboardDataDetails.cpp
+++ b/Engine/Source/Editor/BehaviorTreeEditor/Private/DetailCustomizations/BlackboardDataDetails.cpp
@@ -2,7 +2,7 @@
 
 #include "BehaviorTreeEditorPrivatePCH.h"
 #include "BlackboardDataDetails.h"
-
+#include "BlackboardData.h"
 
 TSharedRef<IDetailCustomization> FBlackboardDataDetails::MakeInstance(FOnGetSelectedBlackboardItemIndex InOnGetSelectedBlackboardItemIndex)
 {
