- Really fix BlackboardDataDetails.cpp non-unity

diff --git a/Engine/Source/Editor/BehaviorTreeEditor/Private/DetailCustomizations/BlackboardDataDetails.cpp b/Engine/Source/Editor/BehaviorTreeEditor/Private/DetailCustomizations/BlackboardDataDetails.cpp
--- a/Engine/Source/Editor/BehaviorTreeEditor/Private/DetailCustomizations/BlackboardDataDetails.cpp
+++ b/Engine/Source/Editor/BehaviorTreeEditor/Private/DetailCustomizations/BlackboardDataDetails.cpp
@@ -2,7 +2,7 @@
 
 #include "BehaviorTreeEditorPrivatePCH.h"
 #include "BlackboardDataDetails.h"
-#include "BlackboardData.h"
+#include "Runtime/AIModule/Classes/BehaviorTree/BlackboardData.h"
 
 TSharedRef<IDetailCustomization> FBlackboardDataDetails::MakeInstance(FOnGetSelectedBlackboardItemIndex InOnGetSelectedBlackboardItemIndex)
 {
