Prevented crash when resetting Blackboard key type

Some parts of the code were still assuming that the blackboard key type was always valid. Fixed up & tested.

TTP# 343906 -  EDITOR: CRASH: In a blackboard asset, resetting a key type to default results in a crash

diff --git a/Engine/Source/Editor/BehaviorTreeEditor/Private/SBehaviorTreeBlackboardView.cpp b/Engine/Source/Editor/BehaviorTreeEditor/Private/SBehaviorTreeBlackboardView.cpp
--- a/Engine/Source/Editor/BehaviorTreeEditor/Private/SBehaviorTreeBlackboardView.cpp
+++ b/Engine/Source/Editor/BehaviorTreeEditor/Private/SBehaviorTreeBlackboardView.cpp
@@ -44,7 +44,7 @@ public:
 	void Update()
 	{
 		MenuDescription = FText::FromName(Key.EntryName);
-		TooltipDescription = FText::Format(LOCTEXT("BlackboardEntryFormat", "{0} '{1}'"), Key.KeyType->GetClass()->GetDisplayNameText(), FText::FromName(Key.EntryName)).ToString();
+		TooltipDescription = FText::Format(LOCTEXT("BlackboardEntryFormat", "{0} '{1}'"), Key.KeyType ? Key.KeyType->GetClass()->GetDisplayNameText() : LOCTEXT("NullKeyDesc", "None"), FText::FromName(Key.EntryName)).ToString();
 		SectionID = bIsInherited ? EBlackboardSectionTitles::InheritedKeys : EBlackboardSectionTitles::Keys;	
 	}
 
@@ -127,9 +127,11 @@ private:
 		check(InGraphAction.IsValid());
 		check(InGraphAction->GetTypeId() == FEdGraphSchemaAction_BlackboardEntry::StaticGetTypeId());
 		TSharedPtr<FEdGraphSchemaAction_BlackboardEntry> BlackboardEntryAction = StaticCastSharedPtr<FEdGraphSchemaAction_BlackboardEntry>(InGraphAction);
-		check(BlackboardEntryAction->Key.KeyType);
 
-		OutIconBrush = FClassIconFinder::FindIconForClass(BlackboardEntryAction->Key.KeyType->GetClass());
+		if(BlackboardEntryAction->Key.KeyType)
+		{
+			OutIconBrush = FClassIconFinder::FindIconForClass(BlackboardEntryAction->Key.KeyType->GetClass());
+		}
 	}
 
 	virtual TSharedRef<SWidget> CreateTextSlotWidget( const FSlateFontInfo& NameFont, FCreateWidgetForActionData* const InCreateData, bool bInIsReadOnly ) override
