Autowire new Behavior Tree nodes when dragging off from input an pin

Extend the autowiring of new nodes to deal with draggin from input nodes, rather than mjust output as before.

TTP# 344678 - EDITOR: BEHAVIOR TREES: Nodes are not automatically linked when one is create from input bar drag

diff --git a/Engine/Source/Editor/BehaviorTreeEditor/Private/BehaviorTreeGraphNode.cpp b/Engine/Source/Editor/BehaviorTreeEditor/Private/BehaviorTreeGraphNode.cpp
--- a/Engine/Source/Editor/BehaviorTreeEditor/Private/BehaviorTreeGraphNode.cpp
+++ b/Engine/Source/Editor/BehaviorTreeEditor/Private/BehaviorTreeGraphNode.cpp
@@ -231,6 +231,10 @@ void UBehaviorTreeGraphNode::AutowireNewNode(UEdGraphPin* FromPin)
 		{
 			FromPin->GetOwningNode()->NodeConnectionListChanged();
 		}
+		else if(GetSchema()->TryCreateConnection(GetOutputPin(), FromPin))
+		{
+			NodeConnectionListChanged();
+		}
 	}
 }
 
