Update Blackboard view when Blackboard is edited in Blackboard mode

Made sure the Blackboard view is up-to-date when Blackboard keys are added, removed and renamed.

TTP# 344116 - AI: BehaviorTree: The blackboard view tab does not update when changes are made to keys

diff --git a/Engine/Source/Editor/BehaviorTreeEditor/Private/BehaviorTreeEditor.cpp b/Engine/Source/Editor/BehaviorTreeEditor/Private/BehaviorTreeEditor.cpp
--- a/Engine/Source/Editor/BehaviorTreeEditor/Private/BehaviorTreeEditor.cpp
+++ b/Engine/Source/Editor/BehaviorTreeEditor/Private/BehaviorTreeEditor.cpp
@@ -247,7 +247,8 @@ void FBehaviorTreeEditor::InitBehaviorTreeEditor( const EToolkitMode::Type Mode,
 				.OnIsDebuggerReady(this, &FBehaviorTreeEditor::IsDebuggerReady)
 				.OnIsDebuggerPaused(this, &FBehaviorTreeEditor::IsDebuggerPaused)
 				.OnGetDebugTimeStamp(this, &FBehaviorTreeEditor::HandleGetDebugTimeStamp)
-				.OnGetDisplayCurrentState(this, &FBehaviorTreeEditor::HandleGetDisplayCurrentState);
+				.OnGetDisplayCurrentState(this, &FBehaviorTreeEditor::HandleGetDisplayCurrentState)
+				.OnBlackboardKeyChanged(this, &FBehaviorTreeEditor::HandleBlackboardKeyChanged);
 	}
 	else
 	{
@@ -387,6 +388,15 @@ bool FBehaviorTreeEditor::HandleGetDisplayCurrentState() const
 	return false;
 }
 
+void FBehaviorTreeEditor::HandleBlackboardKeyChanged(UBlackboardData* InBlackboardData, FBlackboardEntry* const InKey)
+{
+	if(BlackboardView.IsValid())
+	{
+		// re-set object in blackboard view to keep it up to date
+		BlackboardView->SetObject(InBlackboardData);
+	}
+}
+
 void FBehaviorTreeEditor::GetBlackboardSelectionInfo(int32& OutSelectionIndex, bool& bOutIsInherited) const
 {
 	OutSelectionIndex = CurrentBlackboardEntryIndex;
diff --git a/Engine/Source/Editor/BehaviorTreeEditor/Private/BehaviorTreeEditor.h b/Engine/Source/Editor/BehaviorTreeEditor/Private/BehaviorTreeEditor.h
--- a/Engine/Source/Editor/BehaviorTreeEditor/Private/BehaviorTreeEditor.h
+++ b/Engine/Source/Editor/BehaviorTreeEditor/Private/BehaviorTreeEditor.h
@@ -181,6 +181,9 @@ public:
 	/** Delegate handler for determining whether to display the current state */
 	bool HandleGetDisplayCurrentState() const;
 
+	/** Delegate handler for when a blackboard key changes (added, removed, renamed) */
+	void HandleBlackboardKeyChanged(UBlackboardData* InBlackboardData, FBlackboardEntry* const InKey);
+
 	/** Get the currently selected blackboard entry */
 	void GetBlackboardSelectionInfo(int32& OutSelectionIndex, bool& bOutIsInherited) const;
 
diff --git a/Engine/Source/Editor/BehaviorTreeEditor/Private/SBehaviorTreeBlackboardEditor.cpp b/Engine/Source/Editor/BehaviorTreeEditor/Private/SBehaviorTreeBlackboardEditor.cpp
--- a/Engine/Source/Editor/BehaviorTreeEditor/Private/SBehaviorTreeBlackboardEditor.cpp
+++ b/Engine/Source/Editor/BehaviorTreeEditor/Private/SBehaviorTreeBlackboardEditor.cpp
@@ -48,6 +48,7 @@ void SBehaviorTreeBlackboardEditor::Construct(const FArguments& InArgs, TSharedR
 		.OnIsDebuggerReady(InArgs._OnIsDebuggerReady)
 		.OnIsDebuggerPaused(InArgs._OnIsDebuggerPaused)
 		.OnGetDebugTimeStamp(InArgs._OnGetDebugTimeStamp)
+		.OnBlackboardKeyChanged(InArgs._OnBlackboardKeyChanged)
 		.IsReadOnly(false),
 		CommandList,
 		InBlackboardData
@@ -107,6 +108,7 @@ void SBehaviorTreeBlackboardEditor::HandleDeleteEntry()
 			}
 
 			GraphActionMenu->RefreshAllActions(true);
+			OnBlackboardKeyChanged.ExecuteIfBound(BlackboardData, nullptr);
 
 			// signal de-selection
 			if(OnEntrySelected.IsBound())
@@ -207,6 +209,7 @@ void SBehaviorTreeBlackboardEditor::HandleKeyClassPicked(UClass* InClass)
 	BlackboardData->Keys.Add(Entry);
 
 	GraphActionMenu->RefreshAllActions(true);
+	OnBlackboardKeyChanged.ExecuteIfBound(BlackboardData, &BlackboardData->Keys.Last());
 
 	GraphActionMenu->SelectItemByName(Entry.EntryName, ESelectInfo::OnMouseClick);
 	GraphActionMenu->OnRequestRenameOnActionNode();
diff --git a/Engine/Source/Editor/BehaviorTreeEditor/Private/SBehaviorTreeBlackboardEditor.h b/Engine/Source/Editor/BehaviorTreeEditor/Private/SBehaviorTreeBlackboardEditor.h
--- a/Engine/Source/Editor/BehaviorTreeEditor/Private/SBehaviorTreeBlackboardEditor.h
+++ b/Engine/Source/Editor/BehaviorTreeEditor/Private/SBehaviorTreeBlackboardEditor.h
@@ -18,6 +18,7 @@ public:
 		SLATE_EVENT(FOnIsDebuggerReady, OnIsDebuggerReady)
 		SLATE_EVENT(FOnIsDebuggerPaused, OnIsDebuggerPaused)
 		SLATE_EVENT(FOnGetDebugTimeStamp, OnGetDebugTimeStamp)
+		SLATE_EVENT(FOnBlackboardKeyChanged, OnBlackboardKeyChanged)
 
 	SLATE_END_ARGS()
 
diff --git a/Engine/Source/Editor/BehaviorTreeEditor/Private/SBehaviorTreeBlackboardView.cpp b/Engine/Source/Editor/BehaviorTreeEditor/Private/SBehaviorTreeBlackboardView.cpp
--- a/Engine/Source/Editor/BehaviorTreeEditor/Private/SBehaviorTreeBlackboardView.cpp
+++ b/Engine/Source/Editor/BehaviorTreeEditor/Private/SBehaviorTreeBlackboardView.cpp
@@ -63,6 +63,7 @@ class SBehaviorTreeBlackboardItem : public SGraphPaletteItem
 		SLATE_EVENT(FOnGetDebugKeyValue, OnGetDebugKeyValue)
 		SLATE_EVENT(FOnGetDisplayCurrentState, OnGetDisplayCurrentState)
 		SLATE_EVENT(FOnIsDebuggerReady, OnIsDebuggerReady)
+		SLATE_EVENT(FOnBlackboardKeyChanged, OnBlackboardKeyChanged)
 
 	SLATE_END_ARGS()
 
@@ -71,6 +72,7 @@ class SBehaviorTreeBlackboardItem : public SGraphPaletteItem
 		OnGetDebugKeyValue = InArgs._OnGetDebugKeyValue;
 		OnIsDebuggerReady = InArgs._OnIsDebuggerReady;
 		OnGetDisplayCurrentState = InArgs._OnGetDisplayCurrentState;
+		OnBlackboardKeyChanged = InArgs._OnBlackboardKeyChanged;
 
 		const FSlateFontInfo NameFont = FSlateFontInfo( FPaths::EngineContentDir() / TEXT("Slate/Fonts/Roboto-Regular.ttf"), 10);
 
@@ -186,6 +188,8 @@ private:
 			BlackboardEntryAction->Key.EntryName = FName(*NewText.ToString());
 
 			BlackboardEntryAction->Update();
+
+			OnBlackboardKeyChanged.ExecuteIfBound(BlackboardEntryAction->BlackboardData, &BlackboardEntryAction->Key);
 		}
 	}
 
@@ -270,6 +274,9 @@ private:
 	/** Delegate used to determine whether the BT debugger displaying the current state */
 	FOnGetDisplayCurrentState OnGetDisplayCurrentState;
 
+	/** Delegate for when a blackboard key changes (added, removed, renamed) */
+	FOnBlackboardKeyChanged OnBlackboardKeyChanged;
+
 	/** Read-only flag */
 	bool bIsReadOnly;
 };
@@ -291,6 +298,7 @@ void SBehaviorTreeBlackboardView::Construct(const FArguments& InArgs, TSharedRef
 	OnIsDebuggerPaused = InArgs._OnIsDebuggerPaused;
 	OnGetDebugTimeStamp = InArgs._OnGetDebugTimeStamp;
 	OnGetDisplayCurrentState = InArgs._OnGetDisplayCurrentState;
+	OnBlackboardKeyChanged = InArgs._OnBlackboardKeyChanged;
 
 	BlackboardData = InBlackboardData;
 
@@ -377,7 +385,8 @@ TSharedRef<SWidget> SBehaviorTreeBlackboardView::HandleCreateWidgetForAction(FCr
 	return SNew(SBehaviorTreeBlackboardItem, InCreateData)
 		.OnIsDebuggerReady(OnIsDebuggerReady)
 		.OnGetDebugKeyValue(OnGetDebugKeyValue)
-		.OnGetDisplayCurrentState(this, &SBehaviorTreeBlackboardView::IsUsingCurrentValues);
+		.OnGetDisplayCurrentState(this, &SBehaviorTreeBlackboardView::IsUsingCurrentValues)
+		.OnBlackboardKeyChanged(OnBlackboardKeyChanged);
 }
 
 void SBehaviorTreeBlackboardView::HandleCollectAllActions( FGraphActionListBuilderBase& GraphActionListBuilder )
diff --git a/Engine/Source/Editor/BehaviorTreeEditor/Private/SBehaviorTreeBlackboardView.h b/Engine/Source/Editor/BehaviorTreeEditor/Private/SBehaviorTreeBlackboardView.h
--- a/Engine/Source/Editor/BehaviorTreeEditor/Private/SBehaviorTreeBlackboardView.h
+++ b/Engine/Source/Editor/BehaviorTreeEditor/Private/SBehaviorTreeBlackboardView.h
@@ -27,6 +27,8 @@ DECLARE_DELEGATE_RetVal(bool, FOnGetDisplayCurrentState);
 /** Delegate used to get the debugger's current timestamp */
 DECLARE_DELEGATE_RetVal_OneParam(float, FOnGetDebugTimeStamp, bool /* bUseCurrentState */);
 
+/** Delegate for when a blackboard key changes (added, removed, renamed) */
+DECLARE_DELEGATE_TwoParams(FOnBlackboardKeyChanged, UBlackboardData* /*InBlackboardData*/, FBlackboardEntry* const /*InKey*/);
 
 /** Displays blackboard entries */
 class SBehaviorTreeBlackboardView : public SCompoundWidget, public FGCObject
@@ -43,6 +45,7 @@ public:
 		SLATE_EVENT(FOnIsDebuggerReady, OnIsDebuggerReady)
 		SLATE_EVENT(FOnIsDebuggerPaused, OnIsDebuggerPaused)
 		SLATE_EVENT(FOnGetDebugTimeStamp, OnGetDebugTimeStamp)
+		SLATE_EVENT(FOnBlackboardKeyChanged, OnBlackboardKeyChanged)
 		SLATE_ARGUMENT(bool, IsReadOnly)
 
 	SLATE_END_ARGS()
@@ -149,6 +152,9 @@ protected:
 	/** Delegate used to get the debugger's current timestamp */
 	FOnGetDebugTimeStamp OnGetDebugTimeStamp;
 
+	/** Delegate for when a blackboard key changes (added, removed, renamed) */
+	FOnBlackboardKeyChanged OnBlackboardKeyChanged;
+
 	/** Whether we want to show the current or saved state */
 	bool bShowCurrentState;
 };
\ No newline at end of file
