Blackboard mode is now selectable after undoing a reset of the property

The cached blackboard asset was not being updated on Undo, only on property change. Added extra PostEditUndo() call and prevented a full rebuild if the blackboard had not changed,

TTP# 344015 - AI: BehaviorTree: Possible to have the blackboard mode button grayed out even though a blackboard asset is set in the behavior tree

diff --git a/Engine/Source/Editor/BehaviorTreeEditor/Classes/BehaviorTreeGraphNode_Root.h b/Engine/Source/Editor/BehaviorTreeEditor/Classes/BehaviorTreeGraphNode_Root.h
--- a/Engine/Source/Editor/BehaviorTreeEditor/Classes/BehaviorTreeGraphNode_Root.h
+++ b/Engine/Source/Editor/BehaviorTreeEditor/Classes/BehaviorTreeGraphNode_Root.h
@@ -24,6 +24,7 @@ class UBehaviorTreeGraphNode_Root : public UBehaviorTreeGraphNode
 	virtual FString GetTooltip() const override;
 
 	virtual void PostEditChangeProperty( struct FPropertyChangedEvent& PropertyChangedEvent) override;
+	virtual void PostEditUndo() override;
 	virtual FString	GetDescription() const override;
 
 	/** notify behavior tree about blackboard change */
diff --git a/Engine/Source/Editor/BehaviorTreeEditor/Private/BehaviorTreeGraphNode_Root.cpp b/Engine/Source/Editor/BehaviorTreeEditor/Private/BehaviorTreeGraphNode_Root.cpp
--- a/Engine/Source/Editor/BehaviorTreeEditor/Private/BehaviorTreeGraphNode_Root.cpp
+++ b/Engine/Source/Editor/BehaviorTreeEditor/Private/BehaviorTreeGraphNode_Root.cpp
@@ -56,6 +56,12 @@ void UBehaviorTreeGraphNode_Root::PostEditChangeProperty( struct FPropertyChange
 	}
 }
 
+void UBehaviorTreeGraphNode_Root::PostEditUndo()
+{
+	Super::PostEditUndo();
+	UpdateBlackboard();
+}
+
 FString	UBehaviorTreeGraphNode_Root::GetDescription() const
 {
 	return *GetNameSafe(BlackboardAsset);
@@ -65,7 +71,7 @@ void UBehaviorTreeGraphNode_Root::UpdateBlackboard()
 {
 	UBehaviorTreeGraph* MyGraph = GetBehaviorTreeGraph();
 	UBehaviorTree* BTAsset = Cast<UBehaviorTree>(MyGraph->GetOuter());
-	if (BTAsset)
+	if (BTAsset && BTAsset->BlackboardAsset != BlackboardAsset)
 	{
 		BTAsset->BlackboardAsset = BlackboardAsset;
 		MyGraph->UpdateBlackboardChange();
