Deprecating EdGraphNode's GetTooltip() (in exchange for one that returns FText over FString).

diff --git a/Engine/Plugins/Experimental/LiveEditor/Source/LiveEditor/Private/LiveEditor.cpp b/Engine/Plugins/Experimental/LiveEditor/Source/LiveEditor/Private/LiveEditor.cpp
--- a/Engine/Plugins/Experimental/LiveEditor/Source/LiveEditor/Private/LiveEditor.cpp
+++ b/Engine/Plugins/Experimental/LiveEditor/Source/LiveEditor/Private/LiveEditor.cpp
@@ -100,7 +100,7 @@ void FLiveEditorNodeInjector::AddContextMenuNodes(FBlueprintGraphActionListBuild
 {
 	UK2Node* NodeTemplate = ContextMenuBuilder.CreateTemplateNode<UK2Node_LiveEditObject>();
 	FString FunctionCategory( TEXT("LiveEditor") );
-	TSharedPtr<FEdGraphSchemaAction_K2NewNode> Action = FK2ActionMenuBuilder::AddNewNodeAction(ContextMenuBuilder, FunctionCategory, NSLOCTEXT("LiveEditorPlugin", "LiveEditObject", "LiveEdit Object"), NodeTemplate->GetTooltip(), 0, NodeTemplate->GetKeywords());
+	TSharedPtr<FEdGraphSchemaAction_K2NewNode> Action = FK2ActionMenuBuilder::AddNewNodeAction(ContextMenuBuilder, FunctionCategory, NSLOCTEXT("LiveEditorPlugin", "LiveEditObject", "LiveEdit Object"), NodeTemplate->GetTooltipText().ToString(), 0, NodeTemplate->GetKeywords());
 	Action->NodeTemplate = NodeTemplate;
 }
 
diff --git a/Engine/Source/Editor/AnimGraph/Classes/AnimGraphNode_ApplyAdditive.h b/Engine/Source/Editor/AnimGraph/Classes/AnimGraphNode_ApplyAdditive.h
--- a/Engine/Source/Editor/AnimGraph/Classes/AnimGraphNode_ApplyAdditive.h
+++ b/Engine/Source/Editor/AnimGraph/Classes/AnimGraphNode_ApplyAdditive.h
@@ -14,7 +14,7 @@ class UAnimGraphNode_ApplyAdditive : public UAnimGraphNode_Base
 	FAnimNode_ApplyAdditive Node;
 
 	// UEdGraphNode interface
-	virtual FString GetTooltip() const override;
+	virtual FText GetTooltipText() const override;
 	virtual FText GetNodeTitle(ENodeTitleType::Type TitleType) const override;
 	virtual FLinearColor GetNodeTitleColor() const override;
 	// End of UEdGraphNode interface
diff --git a/Engine/Source/Editor/AnimGraph/Classes/AnimGraphNode_BlendListByBool.h b/Engine/Source/Editor/AnimGraph/Classes/AnimGraphNode_BlendListByBool.h
--- a/Engine/Source/Editor/AnimGraph/Classes/AnimGraphNode_BlendListByBool.h
+++ b/Engine/Source/Editor/AnimGraph/Classes/AnimGraphNode_BlendListByBool.h
@@ -14,7 +14,7 @@ class UAnimGraphNode_BlendListByBool : public UAnimGraphNode_BlendListBase
 	FAnimNode_BlendListByBool Node;
 
 	// UEdGraphNode interface
-	virtual FString GetTooltip() const override;
+	virtual FText GetTooltipText() const override;
 	virtual FText GetNodeTitle(ENodeTitleType::Type TitleType) const override;
 	// End of UEdGraphNode interface
 
diff --git a/Engine/Source/Editor/AnimGraph/Classes/AnimGraphNode_SkeletalControlBase.h b/Engine/Source/Editor/AnimGraph/Classes/AnimGraphNode_SkeletalControlBase.h
--- a/Engine/Source/Editor/AnimGraph/Classes/AnimGraphNode_SkeletalControlBase.h
+++ b/Engine/Source/Editor/AnimGraph/Classes/AnimGraphNode_SkeletalControlBase.h
@@ -13,7 +13,7 @@ class UAnimGraphNode_SkeletalControlBase : public UAnimGraphNode_Base
 public:
 	// UEdGraphNode interface
 	ANIMGRAPH_API virtual FLinearColor GetNodeTitleColor() const override;
-	ANIMGRAPH_API virtual FString GetTooltip() const override;
+	ANIMGRAPH_API virtual FText GetTooltipText() const override;
 	// End of UEdGraphNode interface
 
 	// UAnimGraphNode_Base interface
diff --git a/Engine/Source/Editor/AnimGraph/Private/AnimGraphNode_ApplyAdditive.cpp b/Engine/Source/Editor/AnimGraph/Private/AnimGraphNode_ApplyAdditive.cpp
--- a/Engine/Source/Editor/AnimGraph/Private/AnimGraphNode_ApplyAdditive.cpp
+++ b/Engine/Source/Editor/AnimGraph/Private/AnimGraphNode_ApplyAdditive.cpp
@@ -21,9 +21,9 @@ FLinearColor UAnimGraphNode_ApplyAdditive::GetNodeTitleColor() const
 	return FLinearColor(0.75f, 0.75f, 0.75f);
 }
 
-FString UAnimGraphNode_ApplyAdditive::GetTooltip() const
+FText UAnimGraphNode_ApplyAdditive::GetTooltipText() const
 {
-	return LOCTEXT("AnimGraphNode_ApplyAdditive_Tooltip", "Apply additive animation to normal pose").ToString();
+	return LOCTEXT("AnimGraphNode_ApplyAdditive_Tooltip", "Apply additive animation to normal pose");
 }
 
 FText UAnimGraphNode_ApplyAdditive::GetNodeTitle(ENodeTitleType::Type TitleType) const
diff --git a/Engine/Source/Editor/AnimGraph/Private/AnimGraphNode_Base.cpp b/Engine/Source/Editor/AnimGraph/Private/AnimGraphNode_Base.cpp
--- a/Engine/Source/Editor/AnimGraph/Private/AnimGraphNode_Base.cpp
+++ b/Engine/Source/Editor/AnimGraph/Private/AnimGraphNode_Base.cpp
@@ -246,7 +246,7 @@ TSharedPtr<FEdGraphSchemaAction_K2NewNode> UAnimGraphNode_Base::CreateDefaultMen
 
 	FString Category = TemplateNode->GetNodeCategory();
 	FText MenuDesc = TemplateNode->GetNodeTitle(ENodeTitleType::ListView);
-	FString Tooltip = TemplateNode->GetTooltip();
+	FString Tooltip = TemplateNode->GetTooltipText().ToString();
 	FString Keywords = TemplateNode->GetKeywords();
 
 	TSharedPtr<FEdGraphSchemaAction_K2NewNode> NodeAction = FK2ActionMenuBuilder::AddNewNodeAction(ContextMenuBuilder, Category, MenuDesc, Tooltip, 0, Keywords);
diff --git a/Engine/Source/Editor/AnimGraph/Private/AnimGraphNode_BlendListByBool.cpp b/Engine/Source/Editor/AnimGraph/Private/AnimGraphNode_BlendListByBool.cpp
--- a/Engine/Source/Editor/AnimGraph/Private/AnimGraphNode_BlendListByBool.cpp
+++ b/Engine/Source/Editor/AnimGraph/Private/AnimGraphNode_BlendListByBool.cpp
@@ -21,9 +21,9 @@ FText UAnimGraphNode_BlendListByBool::GetNodeTitle(ENodeTitleType::Type TitleTyp
 	return LOCTEXT("UAnimGraphNode_BlendListByBool_Tooltip", "Blend Poses by bool");
 }
 
-FString UAnimGraphNode_BlendListByBool::GetTooltip() const
+FText UAnimGraphNode_BlendListByBool::GetTooltipText() const
 {
-	return LOCTEXT("UAnimGraphNode_BlendListByBool_Title", "Blend List (by bool)").ToString();
+	return LOCTEXT("UAnimGraphNode_BlendListByBool_Title", "Blend List (by bool)");
 }
 
 void UAnimGraphNode_BlendListByBool::CustomizePinData(UEdGraphPin* Pin, FName SourcePropertyName, int32 ArrayIndex) const 
diff --git a/Engine/Source/Editor/AnimGraph/Private/AnimGraphNode_BlendListByEnum.cpp b/Engine/Source/Editor/AnimGraph/Private/AnimGraphNode_BlendListByEnum.cpp
--- a/Engine/Source/Editor/AnimGraph/Private/AnimGraphNode_BlendListByEnum.cpp
+++ b/Engine/Source/Editor/AnimGraph/Private/AnimGraphNode_BlendListByEnum.cpp
@@ -24,17 +24,17 @@ FString UAnimGraphNode_BlendListByEnum::GetNodeCategory() const
 	return FString::Printf(TEXT("%s, Blend List by enum"), *Super::GetNodeCategory() );
 }
 
-FString UAnimGraphNode_BlendListByEnum::GetTooltip() const
+FText UAnimGraphNode_BlendListByEnum::GetTooltipText() const
 {
 	if (BoundEnum != NULL)
 	{
 		FFormatNamedArguments Args;
 		Args.Add(TEXT("EnumName"), FText::FromString(BoundEnum->GetName()));
-		return FText::Format(LOCTEXT("AnimGraphNode_BlendListByEnum_Tooltip", "Blend Poses based on enum ({EnumName}). Use context menu to add new entries."), Args).ToString();
+		return FText::Format(LOCTEXT("AnimGraphNode_BlendListByEnum_Tooltip", "Blend Poses based on enum ({EnumName}). Use context menu to add new entries."), Args);
 	}
 	else
 	{
-		return LOCTEXT("AnimGraphNode_BlendListByEnum_TooltipError", "ERROR: Blend Poses (by missing enum)").ToString();
+		return LOCTEXT("AnimGraphNode_BlendListByEnum_TooltipError", "ERROR: Blend Poses (by missing enum)");
 	}
 }
 
@@ -73,7 +73,7 @@ void UAnimGraphNode_BlendListByEnum::GetMenuEntries(FGraphContextMenuBuilder& Co
 			UAnimGraphNode_BlendListByEnum* EnumTemplate = NewObject<UAnimGraphNode_BlendListByEnum>();
 			EnumTemplate->BoundEnum = CurrentEnum;
 
-			TSharedPtr<FEdGraphSchemaAction_K2NewNode> Action = FK2ActionMenuBuilder::AddNewNodeAction(ContextMenuBuilder, EnumTemplate->GetNodeCategory(), EnumTemplate->GetNodeTitle(ENodeTitleType::ListView), EnumTemplate->GetTooltip(), 0, EnumTemplate->GetKeywords());
+			TSharedPtr<FEdGraphSchemaAction_K2NewNode> Action = FK2ActionMenuBuilder::AddNewNodeAction(ContextMenuBuilder, EnumTemplate->GetNodeCategory(), EnumTemplate->GetNodeTitle(ENodeTitleType::ListView), EnumTemplate->GetTooltipText().ToString(), 0, EnumTemplate->GetKeywords());
 			Action->NodeTemplate = EnumTemplate;
 		}
 	}
diff --git a/Engine/Source/Editor/AnimGraph/Private/AnimGraphNode_BlendListByInt.cpp b/Engine/Source/Editor/AnimGraph/Private/AnimGraphNode_BlendListByInt.cpp
--- a/Engine/Source/Editor/AnimGraph/Private/AnimGraphNode_BlendListByInt.cpp
+++ b/Engine/Source/Editor/AnimGraph/Private/AnimGraphNode_BlendListByInt.cpp
@@ -17,9 +17,9 @@ UAnimGraphNode_BlendListByInt::UAnimGraphNode_BlendListByInt(const FPostConstruc
 {
 }
 
-FString UAnimGraphNode_BlendListByInt::GetTooltip() const
+FText UAnimGraphNode_BlendListByInt::GetTooltipText() const
 {
-	return LOCTEXT("AnimGraphNode_BlendListByInt_Tooltip", "Blend List (by int)").ToString();
+	return LOCTEXT("AnimGraphNode_BlendListByInt_Tooltip", "Blend List (by int)");
 }
 
 FText UAnimGraphNode_BlendListByInt::GetNodeTitle(ENodeTitleType::Type TitleType) const
diff --git a/Engine/Source/Editor/AnimGraph/Private/AnimGraphNode_BlendSpaceEvaluator.cpp b/Engine/Source/Editor/AnimGraph/Private/AnimGraphNode_BlendSpaceEvaluator.cpp
--- a/Engine/Source/Editor/AnimGraph/Private/AnimGraphNode_BlendSpaceEvaluator.cpp
+++ b/Engine/Source/Editor/AnimGraph/Private/AnimGraphNode_BlendSpaceEvaluator.cpp
@@ -15,9 +15,10 @@ UAnimGraphNode_BlendSpaceEvaluator::UAnimGraphNode_BlendSpaceEvaluator(const FPo
 {
 }
 
-FString UAnimGraphNode_BlendSpaceEvaluator::GetTooltip() const
+FText UAnimGraphNode_BlendSpaceEvaluator::GetTooltipText() const
 {
-	return FString::Printf(TEXT("Blendspace Evaluator'%s'"), (Node.BlendSpace != NULL) ? *(Node.BlendSpace->GetPathName()) : TEXT("(None)"));
+	FText const BlendspaceName = (Node.BlendSpace != NULL) ? FText::FromString(Node.BlendSpace->GetPathName()) : LOCTEXT("None", "(None)");
+	return FText::Format(LOCTEXT("BlendspaceEvaluator_Tooltip", "Blendspace Evaluator '{0}'"), BlendspaceName);
 }
 
 FText UAnimGraphNode_BlendSpaceEvaluator::GetNodeTitle(ENodeTitleType::Type TitleType) const
diff --git a/Engine/Source/Editor/AnimGraph/Private/AnimGraphNode_BlendSpacePlayer.cpp b/Engine/Source/Editor/AnimGraph/Private/AnimGraphNode_BlendSpacePlayer.cpp
--- a/Engine/Source/Editor/AnimGraph/Private/AnimGraphNode_BlendSpacePlayer.cpp
+++ b/Engine/Source/Editor/AnimGraph/Private/AnimGraphNode_BlendSpacePlayer.cpp
@@ -17,9 +17,10 @@ UAnimGraphNode_BlendSpacePlayer::UAnimGraphNode_BlendSpacePlayer(const FPostCons
 {
 }
 
-FString UAnimGraphNode_BlendSpacePlayer::GetTooltip() const
+FText UAnimGraphNode_BlendSpacePlayer::GetTooltipText() const
 {
-	return FString::Printf(TEXT("Blendspace Player '%s'"), (Node.BlendSpace != NULL) ? *(Node.BlendSpace->GetPathName()) : TEXT("(None)"));
+	const FText BlendSpaceName((Node.BlendSpace != NULL) ? FText::FromString(Node.BlendSpace->GetName()) : LOCTEXT("None", "(None)"));
+	return FText::Format(LOCTEXT("BlendspacePlayer", "Blendspace Player '{0}'"), BlendSpaceName);
 }
 
 FText UAnimGraphNode_BlendSpacePlayer::GetNodeTitle(ENodeTitleType::Type TitleType) const
@@ -30,7 +31,7 @@ FText UAnimGraphNode_BlendSpacePlayer::GetNodeTitle(ENodeTitleType::Type TitleTy
 	{
 		FFormatNamedArguments Args;
 		Args.Add(TEXT("BlendSpaceName"), BlendSpaceName);
-		return FText::Format(LOCTEXT("BlendspacePlayer", "Blendspace Player'{BlendSpaceName}'"), Args);
+		return FText::Format(LOCTEXT("BlendspacePlayer", "Blendspace Player '{BlendSpaceName}'"), Args);
 	}
 	else
 	{
diff --git a/Engine/Source/Editor/AnimGraph/Private/AnimGraphNode_BoneDrivenController.cpp b/Engine/Source/Editor/AnimGraph/Private/AnimGraphNode_BoneDrivenController.cpp
--- a/Engine/Source/Editor/AnimGraph/Private/AnimGraphNode_BoneDrivenController.cpp
+++ b/Engine/Source/Editor/AnimGraph/Private/AnimGraphNode_BoneDrivenController.cpp
@@ -11,9 +11,9 @@ UAnimGraphNode_BoneDrivenController::UAnimGraphNode_BoneDrivenController(const F
 
 }
 
-FString UAnimGraphNode_BoneDrivenController::GetTooltip() const
+FText UAnimGraphNode_BoneDrivenController::GetTooltipText() const
 {
-	return LOCTEXT("UAnimGraphNode_BoneDrivenController_ToolTip", "Drives the transform of a bone using the transform of another").ToString();
+	return LOCTEXT("UAnimGraphNode_BoneDrivenController_ToolTip", "Drives the transform of a bone using the transform of another");
 }
 
 FText UAnimGraphNode_BoneDrivenController::GetNodeTitle(ENodeTitleType::Type TitleType) const
diff --git a/Engine/Source/Editor/AnimGraph/Private/AnimGraphNode_ComponentToLocalSpace.cpp b/Engine/Source/Editor/AnimGraph/Private/AnimGraphNode_ComponentToLocalSpace.cpp
--- a/Engine/Source/Editor/AnimGraph/Private/AnimGraphNode_ComponentToLocalSpace.cpp
+++ b/Engine/Source/Editor/AnimGraph/Private/AnimGraphNode_ComponentToLocalSpace.cpp
@@ -18,9 +18,9 @@ FLinearColor UAnimGraphNode_ComponentToLocalSpace::GetNodeTitleColor() const
 	return FLinearColor(0.7f, 0.7f, 0.7f);
 }
 
-FString UAnimGraphNode_ComponentToLocalSpace::GetTooltip() const
+FText UAnimGraphNode_ComponentToLocalSpace::GetTooltipText() const
 {
-	return LOCTEXT("AnimGraphNode_ComponentToLocalSpace_Tooltip", "Convert Component Space Pose to Local Pose").ToString();
+	return LOCTEXT("AnimGraphNode_ComponentToLocalSpace_Tooltip", "Convert Component Space Pose to Local Pose");
 }
 
 FText UAnimGraphNode_ComponentToLocalSpace::GetNodeTitle(ENodeTitleType::Type TitleType) const
diff --git a/Engine/Source/Editor/AnimGraph/Private/AnimGraphNode_CopyBone.cpp b/Engine/Source/Editor/AnimGraph/Private/AnimGraphNode_CopyBone.cpp
--- a/Engine/Source/Editor/AnimGraph/Private/AnimGraphNode_CopyBone.cpp
+++ b/Engine/Source/Editor/AnimGraph/Private/AnimGraphNode_CopyBone.cpp
@@ -18,9 +18,9 @@ FText UAnimGraphNode_CopyBone::GetControllerDescription() const
 	return LOCTEXT("CopyBone", "Copy Bone");
 }
 
-FString UAnimGraphNode_CopyBone::GetTooltip() const
+FText UAnimGraphNode_CopyBone::GetTooltipText() const
 {
-	return LOCTEXT("AnimGraphNode_CopyBone_Tooltip", "The Copy Bone control copies the Transform data or any component of it - i.e. Translation, Rotation, or Scale - from one bone to another.").ToString();
+	return LOCTEXT("AnimGraphNode_CopyBone_Tooltip", "The Copy Bone control copies the Transform data or any component of it - i.e. Translation, Rotation, or Scale - from one bone to another.");
 }
 
 FText UAnimGraphNode_CopyBone::GetNodeTitle(ENodeTitleType::Type TitleType) const
diff --git a/Engine/Source/Editor/AnimGraph/Private/AnimGraphNode_CustomTransitionResult.cpp b/Engine/Source/Editor/AnimGraph/Private/AnimGraphNode_CustomTransitionResult.cpp
--- a/Engine/Source/Editor/AnimGraph/Private/AnimGraphNode_CustomTransitionResult.cpp
+++ b/Engine/Source/Editor/AnimGraph/Private/AnimGraphNode_CustomTransitionResult.cpp
@@ -20,9 +20,9 @@ FLinearColor UAnimGraphNode_CustomTransitionResult::GetNodeTitleColor() const
 	return GetDefault<UGraphEditorSettings>()->ResultNodeTitleColor;
 }
 
-FString UAnimGraphNode_CustomTransitionResult::GetTooltip() const
+FText UAnimGraphNode_CustomTransitionResult::GetTooltipText() const
 {
-	return LOCTEXT("AnimGraphNode_CustomTransitionResult_Tooltip", "Result node for a custom transition blend graph").ToString();
+	return LOCTEXT("AnimGraphNode_CustomTransitionResult_Tooltip", "Result node for a custom transition blend graph");
 }
 
 FText UAnimGraphNode_CustomTransitionResult::GetNodeTitle(ENodeTitleType::Type TitleType) const
diff --git a/Engine/Source/Editor/AnimGraph/Private/AnimGraphNode_HandIKRetargeting.cpp b/Engine/Source/Editor/AnimGraph/Private/AnimGraphNode_HandIKRetargeting.cpp
--- a/Engine/Source/Editor/AnimGraph/Private/AnimGraphNode_HandIKRetargeting.cpp
+++ b/Engine/Source/Editor/AnimGraph/Private/AnimGraphNode_HandIKRetargeting.cpp
@@ -18,9 +18,9 @@ FText UAnimGraphNode_HandIKRetargeting::GetControllerDescription() const
 	return LOCTEXT("HandIKRetargeting", "Hand IK Retargeting");
 }
 
-FString UAnimGraphNode_HandIKRetargeting::GetTooltip() const
+FText UAnimGraphNode_HandIKRetargeting::GetTooltipText() const
 {
-	return LOCTEXT("AnimGraphNode_HandIKRetargeting_Tooltip", "Handle re-targeting of IK Bone chain. Moves IK bone chain to meet FK hand bones, based on HandFKWeight (to favor either side). (0 = favor left hand, 1 = favor right hand, 0.5 = equal weight).").ToString();
+	return LOCTEXT("AnimGraphNode_HandIKRetargeting_Tooltip", "Handle re-targeting of IK Bone chain. Moves IK bone chain to meet FK hand bones, based on HandFKWeight (to favor either side). (0 = favor left hand, 1 = favor right hand, 0.5 = equal weight).");
 }
 
 FText UAnimGraphNode_HandIKRetargeting::GetNodeTitle(ENodeTitleType::Type TitleType) const
diff --git a/Engine/Source/Editor/AnimGraph/Private/AnimGraphNode_IdentityPose.cpp b/Engine/Source/Editor/AnimGraph/Private/AnimGraphNode_IdentityPose.cpp
--- a/Engine/Source/Editor/AnimGraph/Private/AnimGraphNode_IdentityPose.cpp
+++ b/Engine/Source/Editor/AnimGraph/Private/AnimGraphNode_IdentityPose.cpp
@@ -14,9 +14,9 @@ UAnimGraphNode_IdentityPose::UAnimGraphNode_IdentityPose(const FPostConstructIni
 	Node.RefPoseType = EIT_Additive;
 }
 
-FString UAnimGraphNode_IdentityPose::GetTooltip() const
+FText UAnimGraphNode_IdentityPose::GetTooltipText() const
 {
-	return LOCTEXT("AnimGraphNode_IdentityPose_Tooltip", "Returns identity pose.").ToString();
+	return LOCTEXT("AnimGraphNode_IdentityPose_Tooltip", "Returns identity pose.");
 }
 
 FText UAnimGraphNode_IdentityPose::GetNodeTitle(ENodeTitleType::Type TitleType) const
diff --git a/Engine/Source/Editor/AnimGraph/Private/AnimGraphNode_LayeredBoneBlend.cpp b/Engine/Source/Editor/AnimGraph/Private/AnimGraphNode_LayeredBoneBlend.cpp
--- a/Engine/Source/Editor/AnimGraph/Private/AnimGraphNode_LayeredBoneBlend.cpp
+++ b/Engine/Source/Editor/AnimGraph/Private/AnimGraphNode_LayeredBoneBlend.cpp
@@ -21,9 +21,9 @@ FLinearColor UAnimGraphNode_LayeredBoneBlend::GetNodeTitleColor() const
 	return FLinearColor(0.2f, 0.8f, 0.2f);
 }
 
-FString UAnimGraphNode_LayeredBoneBlend::GetTooltip() const
+FText UAnimGraphNode_LayeredBoneBlend::GetTooltipText() const
 {
-	return LOCTEXT("AnimGraphNode_LayeredBoneBlend_Tooltip", "Layered blend per bone").ToString();
+	return LOCTEXT("AnimGraphNode_LayeredBoneBlend_Tooltip", "Layered blend per bone");
 }
 
 FText UAnimGraphNode_LayeredBoneBlend::GetNodeTitle(ENodeTitleType::Type TitleType) const
diff --git a/Engine/Source/Editor/AnimGraph/Private/AnimGraphNode_LocalRefPose.cpp b/Engine/Source/Editor/AnimGraph/Private/AnimGraphNode_LocalRefPose.cpp
--- a/Engine/Source/Editor/AnimGraph/Private/AnimGraphNode_LocalRefPose.cpp
+++ b/Engine/Source/Editor/AnimGraph/Private/AnimGraphNode_LocalRefPose.cpp
@@ -14,9 +14,9 @@ UAnimGraphNode_LocalRefPose::UAnimGraphNode_LocalRefPose(const FPostConstructIni
 	Node.RefPoseType = EIT_LocalSpace;
 }
 
-FString UAnimGraphNode_LocalRefPose::GetTooltip() const
+FText UAnimGraphNode_LocalRefPose::GetTooltipText() const
 {
-	return LOCTEXT("UAnimGraphNode_LocalRefPose_Tooltip", "Returns local space reference pose.").ToString();
+	return LOCTEXT("UAnimGraphNode_LocalRefPose_Tooltip", "Returns local space reference pose.");
 }
 
 FText UAnimGraphNode_LocalRefPose::GetNodeTitle(ENodeTitleType::Type TitleType) const
diff --git a/Engine/Source/Editor/AnimGraph/Private/AnimGraphNode_LocalToComponentSpace.cpp b/Engine/Source/Editor/AnimGraph/Private/AnimGraphNode_LocalToComponentSpace.cpp
--- a/Engine/Source/Editor/AnimGraph/Private/AnimGraphNode_LocalToComponentSpace.cpp
+++ b/Engine/Source/Editor/AnimGraph/Private/AnimGraphNode_LocalToComponentSpace.cpp
@@ -21,9 +21,9 @@ FLinearColor UAnimGraphNode_LocalToComponentSpace::GetNodeTitleColor() const
 	return FLinearColor(0.7f, 0.7f, 0.7f);
 }
 
-FString UAnimGraphNode_LocalToComponentSpace::GetTooltip() const
+FText UAnimGraphNode_LocalToComponentSpace::GetTooltipText() const
 {
-	return LOCTEXT("AnimGraphNode_LocalToComponentSpace_Tooltip", "Convert Local Pose to Component Space Pose").ToString();
+	return LOCTEXT("AnimGraphNode_LocalToComponentSpace_Tooltip", "Convert Local Pose to Component Space Pose");
 }
 
 FText UAnimGraphNode_LocalToComponentSpace::GetNodeTitle(ENodeTitleType::Type TitleType) const
diff --git a/Engine/Source/Editor/AnimGraph/Private/AnimGraphNode_LookAt.cpp b/Engine/Source/Editor/AnimGraph/Private/AnimGraphNode_LookAt.cpp
--- a/Engine/Source/Editor/AnimGraph/Private/AnimGraphNode_LookAt.cpp
+++ b/Engine/Source/Editor/AnimGraph/Private/AnimGraphNode_LookAt.cpp
@@ -23,9 +23,9 @@ FString UAnimGraphNode_LookAt::GetKeywords() const
 	return TEXT("Look At, Follow, Trace, Track");
 }
 
-FString UAnimGraphNode_LookAt::GetTooltip() const
+FText UAnimGraphNode_LookAt::GetTooltipText() const
 {
-	return LOCTEXT("AnimGraphNode_LookAt_Tooltip", "This node allow a bone to trace or follow another bone").ToString();
+	return LOCTEXT("AnimGraphNode_LookAt_Tooltip", "This node allow a bone to trace or follow another bone");
 }
 
 FText UAnimGraphNode_LookAt::GetNodeTitle(ENodeTitleType::Type TitleType) const
diff --git a/Engine/Source/Editor/AnimGraph/Private/AnimGraphNode_MeshRefPose.cpp b/Engine/Source/Editor/AnimGraph/Private/AnimGraphNode_MeshRefPose.cpp
--- a/Engine/Source/Editor/AnimGraph/Private/AnimGraphNode_MeshRefPose.cpp
+++ b/Engine/Source/Editor/AnimGraph/Private/AnimGraphNode_MeshRefPose.cpp
@@ -24,9 +24,9 @@ FLinearColor UAnimGraphNode_MeshRefPose::GetNodeTitleColor() const
 	return FColor(200, 100, 100);
 }
 
-FString UAnimGraphNode_MeshRefPose::GetTooltip() const
+FText UAnimGraphNode_MeshRefPose::GetTooltipText() const
 {
-	return LOCTEXT("AnimGraphNode_MeshRefPose_Tooltip", "Returns mesh space reference pose.").ToString();
+	return LOCTEXT("AnimGraphNode_MeshRefPose_Tooltip", "Returns mesh space reference pose.");
 }
 
 FText UAnimGraphNode_MeshRefPose::GetNodeTitle(ENodeTitleType::Type TitleType) const
diff --git a/Engine/Source/Editor/AnimGraph/Private/AnimGraphNode_ModifyBone.cpp b/Engine/Source/Editor/AnimGraph/Private/AnimGraphNode_ModifyBone.cpp
--- a/Engine/Source/Editor/AnimGraph/Private/AnimGraphNode_ModifyBone.cpp
+++ b/Engine/Source/Editor/AnimGraph/Private/AnimGraphNode_ModifyBone.cpp
@@ -23,9 +23,9 @@ FString UAnimGraphNode_ModifyBone::GetKeywords() const
 	return TEXT("Modify, Transform");
 }
 
-FString UAnimGraphNode_ModifyBone::GetTooltip() const
+FText UAnimGraphNode_ModifyBone::GetTooltipText() const
 {
-	return LOCTEXT("AnimGraphNode_ModifyBone_Tooltip", "The Transform Bone node alters the transform - i.e. Translation, Rotation, or Scale - of the bone").ToString();
+	return LOCTEXT("AnimGraphNode_ModifyBone_Tooltip", "The Transform Bone node alters the transform - i.e. Translation, Rotation, or Scale - of the bone");
 }
 
 FText UAnimGraphNode_ModifyBone::GetNodeTitle(ENodeTitleType::Type TitleType) const
diff --git a/Engine/Source/Editor/AnimGraph/Private/AnimGraphNode_Root.cpp b/Engine/Source/Editor/AnimGraph/Private/AnimGraphNode_Root.cpp
--- a/Engine/Source/Editor/AnimGraph/Private/AnimGraphNode_Root.cpp
+++ b/Engine/Source/Editor/AnimGraph/Private/AnimGraphNode_Root.cpp
@@ -53,9 +53,9 @@ FText UAnimGraphNode_Root::GetNodeTitle(ENodeTitleType::Type TitleType) const
 	return LOCTEXT("AnimGraphNodeRoot_Title", "Final Animation Pose");
 }
 
-FString UAnimGraphNode_Root::GetTooltip() const
+FText UAnimGraphNode_Root::GetTooltipText() const
 {
-	return LOCTEXT("AnimGraphNodeRoot_Tooltip", "Wire the final animation pose into this node").ToString();
+	return LOCTEXT("AnimGraphNodeRoot_Tooltip", "Wire the final animation pose into this node");
 }
 
 bool UAnimGraphNode_Root::IsSinkNode() const
diff --git a/Engine/Source/Editor/AnimGraph/Private/AnimGraphNode_RotateRootBone.cpp b/Engine/Source/Editor/AnimGraph/Private/AnimGraphNode_RotateRootBone.cpp
--- a/Engine/Source/Editor/AnimGraph/Private/AnimGraphNode_RotateRootBone.cpp
+++ b/Engine/Source/Editor/AnimGraph/Private/AnimGraphNode_RotateRootBone.cpp
@@ -21,9 +21,9 @@ FLinearColor UAnimGraphNode_RotateRootBone::GetNodeTitleColor() const
 	return FLinearColor(0.7f, 0.7f, 0.7f);
 }
 
-FString UAnimGraphNode_RotateRootBone::GetTooltip() const
+FText UAnimGraphNode_RotateRootBone::GetTooltipText() const
 {
-	return LOCTEXT("RotateRootBone", "Rotate Root Bone").ToString();
+	return LOCTEXT("RotateRootBone", "Rotate Root Bone");
 }
 
 FText UAnimGraphNode_RotateRootBone::GetNodeTitle(ENodeTitleType::Type TitleType) const
diff --git a/Engine/Source/Editor/AnimGraph/Private/AnimGraphNode_RotationMultiplier.cpp b/Engine/Source/Editor/AnimGraph/Private/AnimGraphNode_RotationMultiplier.cpp
--- a/Engine/Source/Editor/AnimGraph/Private/AnimGraphNode_RotationMultiplier.cpp
+++ b/Engine/Source/Editor/AnimGraph/Private/AnimGraphNode_RotationMultiplier.cpp
@@ -17,9 +17,9 @@ FText UAnimGraphNode_RotationMultiplier::GetControllerDescription() const
 	return LOCTEXT("ApplyPercentageOfRotation", "Apply a percentage of Rotation");
 }
 
-FString UAnimGraphNode_RotationMultiplier::GetTooltip() const
+FText UAnimGraphNode_RotationMultiplier::GetTooltipText() const
 {
-	return LOCTEXT("AnimGraphNode_RotationMultiplier_Tooltip", "The Apply a Percentage of Rotation control drives the Rotation of a target bone at some specified percentage of the Rotation of another bone within the Skeleton.").ToString();
+	return LOCTEXT("AnimGraphNode_RotationMultiplier_Tooltip", "The Apply a Percentage of Rotation control drives the Rotation of a target bone at some specified percentage of the Rotation of another bone within the Skeleton.");
 }
 
 FText UAnimGraphNode_RotationMultiplier::GetNodeTitle(ENodeTitleType::Type TitleType) const
diff --git a/Engine/Source/Editor/AnimGraph/Private/AnimGraphNode_RotationOffsetBlendSpace.cpp b/Engine/Source/Editor/AnimGraph/Private/AnimGraphNode_RotationOffsetBlendSpace.cpp
--- a/Engine/Source/Editor/AnimGraph/Private/AnimGraphNode_RotationOffsetBlendSpace.cpp
+++ b/Engine/Source/Editor/AnimGraph/Private/AnimGraphNode_RotationOffsetBlendSpace.cpp
@@ -17,9 +17,10 @@ UAnimGraphNode_RotationOffsetBlendSpace::UAnimGraphNode_RotationOffsetBlendSpace
 {
 }
 
-FString UAnimGraphNode_RotationOffsetBlendSpace::GetTooltip() const
+FText UAnimGraphNode_RotationOffsetBlendSpace::GetTooltipText() const
 {
-	return FString::Printf(TEXT("AimOffset %s"), *(Node.BlendSpace->GetPathName()));
+	const FText BlendSpaceName((Node.BlendSpace != NULL) ? FText::FromString(Node.BlendSpace->GetName()) : LOCTEXT("None", "(None)"));
+	return FText::Format(LOCTEXT("AimOffsetTooltip", "AimOffset '{0}'"), BlendSpaceName);
 }
 
 FText UAnimGraphNode_RotationOffsetBlendSpace::GetNodeTitle(ENodeTitleType::Type TitleType) const
diff --git a/Engine/Source/Editor/AnimGraph/Private/AnimGraphNode_SaveCachedPose.cpp b/Engine/Source/Editor/AnimGraph/Private/AnimGraphNode_SaveCachedPose.cpp
--- a/Engine/Source/Editor/AnimGraph/Private/AnimGraphNode_SaveCachedPose.cpp
+++ b/Engine/Source/Editor/AnimGraph/Private/AnimGraphNode_SaveCachedPose.cpp
@@ -40,9 +40,9 @@ UAnimGraphNode_SaveCachedPose::UAnimGraphNode_SaveCachedPose(const FPostConstruc
 	bCanRenameNode = true;
 }
 
-FString UAnimGraphNode_SaveCachedPose::GetTooltip() const
+FText UAnimGraphNode_SaveCachedPose::GetTooltipText() const
 {
-	return TEXT("Denotes an animation tree that can be referenced elsewhere in the blueprint, which will be evaluated at most once per frame and then cached.");
+	return LOCTEXT("SaveCachedPose_Tooltip", "Denotes an animation tree that can be referenced elsewhere in the blueprint, which will be evaluated at most once per frame and then cached.");
 }
 
 FText UAnimGraphNode_SaveCachedPose::GetNodeTitle(ENodeTitleType::Type TitleType) const
@@ -82,7 +82,7 @@ void UAnimGraphNode_SaveCachedPose::GetMenuEntries(FGraphContextMenuBuilder& Con
 			UAnimGraphNode_SaveCachedPose* SaveCachedPose = NewObject<UAnimGraphNode_SaveCachedPose>();
 			SaveCachedPose->CacheName = TEXT("SavedPose") + FString::FromInt(FMath::Rand());
 
-			TSharedPtr<FEdGraphSchemaAction_K2NewNode> SaveCachedPoseAction = FK2ActionMenuBuilder::AddNewNodeAction(ContextMenuBuilder, GetNodeCategory(), LOCTEXT("NewSaveCachedPose", "New Save cached pose..."), SaveCachedPose->GetTooltip(), 0, SaveCachedPose->GetKeywords());
+			TSharedPtr<FEdGraphSchemaAction_K2NewNode> SaveCachedPoseAction = FK2ActionMenuBuilder::AddNewNodeAction(ContextMenuBuilder, GetNodeCategory(), LOCTEXT("NewSaveCachedPose", "New Save cached pose..."), SaveCachedPose->GetTooltipText().ToString(), 0, SaveCachedPose->GetKeywords());
 			SaveCachedPoseAction->NodeTemplate = SaveCachedPose;
 		}
 	}
diff --git a/Engine/Source/Editor/AnimGraph/Private/AnimGraphNode_SequenceEvaluator.cpp b/Engine/Source/Editor/AnimGraph/Private/AnimGraphNode_SequenceEvaluator.cpp
--- a/Engine/Source/Editor/AnimGraph/Private/AnimGraphNode_SequenceEvaluator.cpp
+++ b/Engine/Source/Editor/AnimGraph/Private/AnimGraphNode_SequenceEvaluator.cpp
@@ -36,15 +36,15 @@ void UAnimGraphNode_SequenceEvaluator::ReplaceReferredAnimations(const TMap<UAni
 	HandleAnimReferenceReplacement(Node.Sequence, ComplexAnimsMap, AnimSequenceMap);
 }
 
-FString UAnimGraphNode_SequenceEvaluator::GetTooltip() const
+FText UAnimGraphNode_SequenceEvaluator::GetTooltipText() const
 {
 	if ((Node.Sequence != NULL) && Node.Sequence->IsValidAdditive())
 	{
-		return FString::Printf(TEXT("Evaluate %s (additive)"), *(Node.Sequence->GetPathName()));
+		return FText::Format(LOCTEXT("SequenceEvaluator_AdditiveTooltip", "Evaluate {0} (additive)"), FText::FromString(Node.Sequence->GetPathName()));
 	}
 	else
 	{
-		return FString::Printf(TEXT("Evaluate %s"), *(Node.Sequence->GetPathName()));
+		return FText::Format(LOCTEXT("SequenceEvaluator_Tooltip", "Evaluate {0}"), FText::FromString(Node.Sequence->GetPathName()));
 	}
 }
 
diff --git a/Engine/Source/Editor/AnimGraph/Private/AnimGraphNode_SequencePlayer.cpp b/Engine/Source/Editor/AnimGraph/Private/AnimGraphNode_SequencePlayer.cpp
--- a/Engine/Source/Editor/AnimGraph/Private/AnimGraphNode_SequencePlayer.cpp
+++ b/Engine/Source/Editor/AnimGraph/Private/AnimGraphNode_SequencePlayer.cpp
@@ -69,10 +69,10 @@ FLinearColor UAnimGraphNode_SequencePlayer::GetNodeTitleColor() const
 	}
 }
 
-FString UAnimGraphNode_SequencePlayer::GetTooltip() const
+FText UAnimGraphNode_SequencePlayer::GetTooltipText() const
 {
 	const bool bAdditive = ((Node.Sequence != NULL) && Node.Sequence->IsValidAdditive());
-	return GetTitleGivenAssetInfo(FText::FromString(Node.Sequence->GetPathName()), bAdditive).ToString();
+	return GetTitleGivenAssetInfo(FText::FromString(Node.Sequence->GetPathName()), bAdditive);
 }
 
 FText UAnimGraphNode_SequencePlayer::GetNodeTitle(ENodeTitleType::Type TitleType) const
diff --git a/Engine/Source/Editor/AnimGraph/Private/AnimGraphNode_SkeletalControlBase.cpp b/Engine/Source/Editor/AnimGraph/Private/AnimGraphNode_SkeletalControlBase.cpp
--- a/Engine/Source/Editor/AnimGraph/Private/AnimGraphNode_SkeletalControlBase.cpp
+++ b/Engine/Source/Editor/AnimGraph/Private/AnimGraphNode_SkeletalControlBase.cpp
@@ -30,9 +30,9 @@ FText UAnimGraphNode_SkeletalControlBase::GetControllerDescription() const
 	return LOCTEXT("ImplementMe", "Implement me");
 }
 
-FString UAnimGraphNode_SkeletalControlBase::GetTooltip() const
+FText UAnimGraphNode_SkeletalControlBase::GetTooltipText() const
 {
-	return GetControllerDescription().ToString();
+	return GetControllerDescription();
 }
 
 void UAnimGraphNode_SkeletalControlBase::CreateOutputPins()
diff --git a/Engine/Source/Editor/AnimGraph/Private/AnimGraphNode_Slot.cpp b/Engine/Source/Editor/AnimGraph/Private/AnimGraphNode_Slot.cpp
--- a/Engine/Source/Editor/AnimGraph/Private/AnimGraphNode_Slot.cpp
+++ b/Engine/Source/Editor/AnimGraph/Private/AnimGraphNode_Slot.cpp
@@ -21,9 +21,9 @@ FLinearColor UAnimGraphNode_Slot::GetNodeTitleColor() const
 	return FLinearColor(0.7f, 0.7f, 0.7f);
 }
 
-FString UAnimGraphNode_Slot::GetTooltip() const
+FText UAnimGraphNode_Slot::GetTooltipText() const
 {
-	return TEXT("Plays animation from code using AnimMontage");
+	return LOCTEXT("AnimSlotNode_Tooltip", "Plays animation from code using AnimMontage");
 }
 
 FText UAnimGraphNode_Slot::GetNodeTitle(ENodeTitleType::Type TitleType) const
diff --git a/Engine/Source/Editor/AnimGraph/Private/AnimGraphNode_SpringBone.cpp b/Engine/Source/Editor/AnimGraph/Private/AnimGraphNode_SpringBone.cpp
--- a/Engine/Source/Editor/AnimGraph/Private/AnimGraphNode_SpringBone.cpp
+++ b/Engine/Source/Editor/AnimGraph/Private/AnimGraphNode_SpringBone.cpp
@@ -18,9 +18,9 @@ FText UAnimGraphNode_SpringBone::GetControllerDescription() const
 	return LOCTEXT("SpringController", "Spring controller");
 }
 
-FString UAnimGraphNode_SpringBone::GetTooltip() const
+FText UAnimGraphNode_SpringBone::GetTooltipText() const
 {
-	return LOCTEXT("AnimGraphNode_SpringBone_Tooltip", "The Spring Controller applies a spring solver that can be used to limit how far a bone can stretch from its reference pose position and apply a force in the opposite direction.").ToString();
+	return LOCTEXT("AnimGraphNode_SpringBone_Tooltip", "The Spring Controller applies a spring solver that can be used to limit how far a bone can stretch from its reference pose position and apply a force in the opposite direction.");
 }
 
 FText UAnimGraphNode_SpringBone::GetNodeTitle(ENodeTitleType::Type TitleType) const
diff --git a/Engine/Source/Editor/AnimGraph/Private/AnimGraphNode_StateMachineBase.cpp b/Engine/Source/Editor/AnimGraph/Private/AnimGraphNode_StateMachineBase.cpp
--- a/Engine/Source/Editor/AnimGraph/Private/AnimGraphNode_StateMachineBase.cpp
+++ b/Engine/Source/Editor/AnimGraph/Private/AnimGraphNode_StateMachineBase.cpp
@@ -52,9 +52,9 @@ FLinearColor UAnimGraphNode_StateMachineBase::GetNodeTitleColor() const
 	return FLinearColor(0.8f, 0.8f, 0.8f);
 }
 
-FString UAnimGraphNode_StateMachineBase::GetTooltip() const
+FText UAnimGraphNode_StateMachineBase::GetTooltipText() const
 {
-	return TEXT("Animation State Machine");
+	return LOCTEXT("StateMachineTooltip", "Animation State Machine");
 }
 
 FText UAnimGraphNode_StateMachineBase::GetNodeTitle(ENodeTitleType::Type TitleType) const
diff --git a/Engine/Source/Editor/AnimGraph/Private/AnimGraphNode_Trail.cpp b/Engine/Source/Editor/AnimGraph/Private/AnimGraphNode_Trail.cpp
--- a/Engine/Source/Editor/AnimGraph/Private/AnimGraphNode_Trail.cpp
+++ b/Engine/Source/Editor/AnimGraph/Private/AnimGraphNode_Trail.cpp
@@ -18,9 +18,9 @@ FText UAnimGraphNode_Trail::GetControllerDescription() const
 	return LOCTEXT("TrailController", "Trail controller");
 }
 
-FString UAnimGraphNode_Trail::GetTooltip() const
+FText UAnimGraphNode_Trail::GetTooltipText() const
 {
-	return LOCTEXT("AnimGraphNode_Trail_Tooltip", "The Trail Controller.").ToString();
+	return LOCTEXT("AnimGraphNode_Trail_Tooltip", "The Trail Controller.");
 }
 
 FText UAnimGraphNode_Trail::GetNodeTitle(ENodeTitleType::Type TitleType) const
diff --git a/Engine/Source/Editor/AnimGraph/Private/AnimGraphNode_TransitionPoseEvaluator.cpp b/Engine/Source/Editor/AnimGraph/Private/AnimGraphNode_TransitionPoseEvaluator.cpp
--- a/Engine/Source/Editor/AnimGraph/Private/AnimGraphNode_TransitionPoseEvaluator.cpp
+++ b/Engine/Source/Editor/AnimGraph/Private/AnimGraphNode_TransitionPoseEvaluator.cpp
@@ -22,9 +22,9 @@ FLinearColor UAnimGraphNode_TransitionPoseEvaluator::GetNodeTitleColor() const
 	return FColor(200, 100, 100);
 }
 
-FString UAnimGraphNode_TransitionPoseEvaluator::GetTooltip() const
+FText UAnimGraphNode_TransitionPoseEvaluator::GetTooltipText() const
 {
-	return (Node.DataSource == EEvaluatorDataSource::EDS_DestinationPose) ? TEXT("Evaluates and returns the pose generated by the destination state of this transition") : TEXT("Evaluates and returns the pose generated by the setup prior to this transition firing");
+	return (Node.DataSource == EEvaluatorDataSource::EDS_DestinationPose) ? LOCTEXT("GetDestinationStatePose_Tooltip", "Evaluates and returns the pose generated by the destination state of this transition") : LOCTEXT("GetSourceStatePose_Tooltip", "Evaluates and returns the pose generated by the setup prior to this transition firing");
 }
 
 FText UAnimGraphNode_TransitionPoseEvaluator::GetNodeTitle(ENodeTitleType::Type TitleType) const
diff --git a/Engine/Source/Editor/AnimGraph/Private/AnimGraphNode_TransitionResult.cpp b/Engine/Source/Editor/AnimGraph/Private/AnimGraphNode_TransitionResult.cpp
--- a/Engine/Source/Editor/AnimGraph/Private/AnimGraphNode_TransitionResult.cpp
+++ b/Engine/Source/Editor/AnimGraph/Private/AnimGraphNode_TransitionResult.cpp
@@ -22,9 +22,9 @@ FLinearColor UAnimGraphNode_TransitionResult::GetNodeTitleColor() const
 	return GetDefault<UGraphEditorSettings>()->ResultNodeTitleColor;
 }
 
-FString UAnimGraphNode_TransitionResult::GetTooltip() const
+FText UAnimGraphNode_TransitionResult::GetTooltipText() const
 {
-	return TEXT("This expression is evaluated to determine if the state transition can be taken");
+	return LOCTEXT("TransitionResultTooltip", "This expression is evaluated to determine if the state transition can be taken");
 }
 
 FText UAnimGraphNode_TransitionResult::GetNodeTitle(ENodeTitleType::Type TitleType) const
diff --git a/Engine/Source/Editor/AnimGraph/Private/AnimGraphNode_TwoBoneIK.cpp b/Engine/Source/Editor/AnimGraph/Private/AnimGraphNode_TwoBoneIK.cpp
--- a/Engine/Source/Editor/AnimGraph/Private/AnimGraphNode_TwoBoneIK.cpp
+++ b/Engine/Source/Editor/AnimGraph/Private/AnimGraphNode_TwoBoneIK.cpp
@@ -18,9 +18,9 @@ FText UAnimGraphNode_TwoBoneIK::GetControllerDescription() const
 	return LOCTEXT("TwoBoneIK", "Two Bone IK");
 }
 
-FString UAnimGraphNode_TwoBoneIK::GetTooltip() const
+FText UAnimGraphNode_TwoBoneIK::GetTooltipText() const
 {
-	return LOCTEXT("AnimGraphNode_TwoBoneIK_Tooltip", "The Two Bone IK control applies an inverse kinematic (IK) solver to a 3-joint chain, such as the limbs of a character.").ToString();
+	return LOCTEXT("AnimGraphNode_TwoBoneIK_Tooltip", "The Two Bone IK control applies an inverse kinematic (IK) solver to a 3-joint chain, such as the limbs of a character.");
 }
 
 FText UAnimGraphNode_TwoBoneIK::GetNodeTitle(ENodeTitleType::Type TitleType) const
diff --git a/Engine/Source/Editor/AnimGraph/Private/AnimGraphNode_TwoWayBlend.cpp b/Engine/Source/Editor/AnimGraph/Private/AnimGraphNode_TwoWayBlend.cpp
--- a/Engine/Source/Editor/AnimGraph/Private/AnimGraphNode_TwoWayBlend.cpp
+++ b/Engine/Source/Editor/AnimGraph/Private/AnimGraphNode_TwoWayBlend.cpp
@@ -23,9 +23,9 @@ FLinearColor UAnimGraphNode_TwoWayBlend::GetNodeTitleColor() const
 	return FLinearColor(0.75f, 0.75f, 0.75f);
 }
 
-FString UAnimGraphNode_TwoWayBlend::GetTooltip() const
+FText UAnimGraphNode_TwoWayBlend::GetTooltipText() const
 {
-	return TEXT("Blend two poses together");
+	return LOCTEXT("TwoWayBlendTooltip", "Blend two poses together");
 }
 
 FText UAnimGraphNode_TwoWayBlend::GetNodeTitle(ENodeTitleType::Type TitleType) const
diff --git a/Engine/Source/Editor/AnimGraph/Private/AnimGraphNode_UseCachedPose.cpp b/Engine/Source/Editor/AnimGraph/Private/AnimGraphNode_UseCachedPose.cpp
--- a/Engine/Source/Editor/AnimGraph/Private/AnimGraphNode_UseCachedPose.cpp
+++ b/Engine/Source/Editor/AnimGraph/Private/AnimGraphNode_UseCachedPose.cpp
@@ -18,9 +18,9 @@ UAnimGraphNode_UseCachedPose::UAnimGraphNode_UseCachedPose(const FPostConstructI
 {
 }
 
-FString UAnimGraphNode_UseCachedPose::GetTooltip() const
+FText UAnimGraphNode_UseCachedPose::GetTooltipText() const
 {
-	return LOCTEXT("AnimGraphNode_UseCachedPose_Tooltip", "References an animation tree elsewhere in the blueprint, which will be evaluated at most once per frame.").ToString();
+	return LOCTEXT("AnimGraphNode_UseCachedPose_Tooltip", "References an animation tree elsewhere in the blueprint, which will be evaluated at most once per frame.");
 }
 
 FText UAnimGraphNode_UseCachedPose::GetNodeTitle(ENodeTitleType::Type TitleType) const
@@ -50,7 +50,7 @@ void UAnimGraphNode_UseCachedPose::GetMenuEntries(FGraphContextMenuBuilder& Cont
 			UAnimGraphNode_UseCachedPose* UseCachedPose = NewObject<UAnimGraphNode_UseCachedPose>();
 			UseCachedPose->NameOfCache = (*NodeIt)->CacheName;
 
-			TSharedPtr<FEdGraphSchemaAction_K2NewNode> UseCachedPoseAction = FK2ActionMenuBuilder::AddNewNodeAction(ContextMenuBuilder, GetNodeCategory(), UseCachedPose->GetNodeTitle(ENodeTitleType::ListView), UseCachedPose->GetTooltip(), 0, UseCachedPose->GetKeywords());
+			TSharedPtr<FEdGraphSchemaAction_K2NewNode> UseCachedPoseAction = FK2ActionMenuBuilder::AddNewNodeAction(ContextMenuBuilder, GetNodeCategory(), UseCachedPose->GetNodeTitle(ENodeTitleType::ListView), UseCachedPose->GetTooltipText().ToString(), 0, UseCachedPose->GetKeywords());
 			UseCachedPoseAction->NodeTemplate = UseCachedPose;
 		}
 	}
diff --git a/Engine/Source/Editor/AnimGraph/Private/AnimGraphNode_WheelHandler.cpp b/Engine/Source/Editor/AnimGraph/Private/AnimGraphNode_WheelHandler.cpp
--- a/Engine/Source/Editor/AnimGraph/Private/AnimGraphNode_WheelHandler.cpp
+++ b/Engine/Source/Editor/AnimGraph/Private/AnimGraphNode_WheelHandler.cpp
@@ -24,9 +24,9 @@ FString UAnimGraphNode_WheelHandler::GetKeywords() const
 	return TEXT("Modify, Wheel, Vehicle");
 }
 
-FString UAnimGraphNode_WheelHandler::GetTooltip() const
+FText UAnimGraphNode_WheelHandler::GetTooltipText() const
 {
-	return LOCTEXT("AnimGraphNode_WheelHandler_Tooltip", "This alters the wheel transform based on set up in Wheeled Vehicle. This only works when the owner is WheeledVehicle.").ToString();
+	return LOCTEXT("AnimGraphNode_WheelHandler_Tooltip", "This alters the wheel transform based on set up in Wheeled Vehicle. This only works when the owner is WheeledVehicle.");
 }
 
 FText UAnimGraphNode_WheelHandler::GetNodeTitle(ENodeTitleType::Type TitleType) const
diff --git a/Engine/Source/Editor/AnimGraph/Private/AnimStateConduitNode.cpp b/Engine/Source/Editor/AnimGraph/Private/AnimStateConduitNode.cpp
--- a/Engine/Source/Editor/AnimGraph/Private/AnimStateConduitNode.cpp
+++ b/Engine/Source/Editor/AnimGraph/Private/AnimStateConduitNode.cpp
@@ -9,6 +9,8 @@
 #include "AnimationConduitGraphSchema.h"
 #include "AnimGraphNode_TransitionResult.h"
 
+#define LOCTEXT_NAMESPACE "AnimStateConduitNode"
+
 /////////////////////////////////////////////////////
 // UAnimStateConduitNode
 
@@ -42,9 +44,9 @@ FText UAnimStateConduitNode::GetNodeTitle(ENodeTitleType::Type TitleType) const
 	return FText::FromString(GetStateName());
 }
 
-FString UAnimStateConduitNode::GetTooltip() const
+FText UAnimStateConduitNode::GetTooltipText() const
 {
-	return TEXT("This is a conduit, which allows specification of a predicate condition for an entire group of transitions");
+	return LOCTEXT("ConduitNodeTooltip", "This is a conduit, which allows specification of a predicate condition for an entire group of transitions");
 }
 
 FString UAnimStateConduitNode::GetStateName() const
@@ -151,3 +153,5 @@ void UAnimStateConduitNode::PostPasteNode()
 	FBlueprintEditorUtils::RenameGraphWithSuggestion(BoundGraph, NameValidator, GetStateName());
 	Super::PostPasteNode();
 }
+
+#undef LOCTEXT_NAMESPACE
diff --git a/Engine/Source/Editor/AnimGraph/Private/AnimStateEntryNode.cpp b/Engine/Source/Editor/AnimGraph/Private/AnimStateEntryNode.cpp
--- a/Engine/Source/Editor/AnimGraph/Private/AnimStateEntryNode.cpp
+++ b/Engine/Source/Editor/AnimGraph/Private/AnimStateEntryNode.cpp
@@ -9,6 +9,8 @@
 #include "EdGraphUtilities.h"
 #include "AnimationStateMachineSchema.h"
 
+#define LOCTEXT_NAMESPACE "AnimStateEntryNode"
+
 /////////////////////////////////////////////////////
 // UAnimStateEntryNode
 
@@ -29,9 +31,9 @@ FText UAnimStateEntryNode::GetNodeTitle(ENodeTitleType::Type TitleType) const
 	return FText::FromString(Graph->GetName());
 }
 
-FString UAnimStateEntryNode::GetTooltip() const
+FText UAnimStateEntryNode::GetTooltipText() const
 {
-	return TEXT("Entry point for state machine");
+	return LOCTEXT("StateEntryNodeTooltip", "Entry point for state machine");
 }
 
 UEdGraphNode* UAnimStateEntryNode::GetOutputNode() const
@@ -47,4 +49,4 @@ UEdGraphNode* UAnimStateEntryNode::GetOutputNode() const
 	return NULL;
 }
 
-
+#undef LOCTEXT_NAMESPACE
diff --git a/Engine/Source/Editor/AnimGraph/Private/AnimStateNode.cpp b/Engine/Source/Editor/AnimGraph/Private/AnimStateNode.cpp
--- a/Engine/Source/Editor/AnimGraph/Private/AnimStateNode.cpp
+++ b/Engine/Source/Editor/AnimGraph/Private/AnimStateNode.cpp
@@ -13,6 +13,8 @@
 #include "AnimGraphNode_StateResult.h"
 #include "AnimationStateGraphSchema.h"
 
+#define LOCTEXT_NAMESPACE "AnimStateNode"
+
 /////////////////////////////////////////////////////
 // UAnimStateNode
 
@@ -47,9 +49,9 @@ FText UAnimStateNode::GetNodeTitle(ENodeTitleType::Type TitleType) const
 	return FText::FromString(GetStateName());
 }
 
-FString UAnimStateNode::GetTooltip() const
+FText UAnimStateNode::GetTooltipText() const
 {
-	return TEXT("This is a state");
+	return LOCTEXT("AnimStateNode_Tooltip", "This is a state");
 }
 
 void UAnimStateNode::PostEditChangeProperty(struct FPropertyChangedEvent& PropertyChangedEvent)
@@ -176,3 +178,5 @@ void UAnimStateNode::DestroyNode()
 		FBlueprintEditorUtils::RemoveGraph(Blueprint, GraphToRemove, EGraphRemoveFlags::Recompile);
 	}
 }
+
+#undef LOCTEXT_NAMESPACE
diff --git a/Engine/Source/Editor/AnimGraph/Private/AnimStateTransitionNode.cpp b/Engine/Source/Editor/AnimGraph/Private/AnimStateTransitionNode.cpp
--- a/Engine/Source/Editor/AnimGraph/Private/AnimStateTransitionNode.cpp
+++ b/Engine/Source/Editor/AnimGraph/Private/AnimStateTransitionNode.cpp
@@ -154,9 +154,9 @@ FText UAnimStateTransitionNode::GetNodeTitle(ENodeTitleType::Type TitleType) con
 	}
 }
 
-FString UAnimStateTransitionNode::GetTooltip() const
+FText UAnimStateTransitionNode::GetTooltipText() const
 {
-	return TEXT("This is a state transition");
+	return LOCTEXT("StateTransitionTooltip", "This is a state transition");
 }
 
 UAnimStateNodeBase* UAnimStateTransitionNode::GetPreviousState() const
diff --git a/Engine/Source/Editor/AnimGraph/Private/AnimationGraphSchema.cpp b/Engine/Source/Editor/AnimGraph/Private/AnimationGraphSchema.cpp
--- a/Engine/Source/Editor/AnimGraph/Private/AnimationGraphSchema.cpp
+++ b/Engine/Source/Editor/AnimGraph/Private/AnimationGraphSchema.cpp
@@ -62,7 +62,7 @@ void UAnimationGraphSchema::GetStateMachineMenuItems(FGraphContextMenuBuilder& C
 			const FString OwnerName = StateNode->GetOuter()->GetName();
 			const FString Title = FString::Printf(TEXT("Current %s for state '%s.%s'"), *UK2Node_TransitionRuleGetter::GetFriendlyName(NodeTemplate->GetterType).ToString(), *OwnerName, *StateNode->GetStateName());
 
-			TSharedPtr<FEdGraphSchemaAction_K2NewNode> Action = FK2ActionMenuBuilder::AddNewNodeAction(ContextMenuBuilder, StateCategory, FText::FromString(Title), NodeTemplate->GetTooltip(), 0, NodeTemplate->GetKeywords());
+			TSharedPtr<FEdGraphSchemaAction_K2NewNode> Action = FK2ActionMenuBuilder::AddNewNodeAction(ContextMenuBuilder, StateCategory, FText::FromString(Title), NodeTemplate->GetTooltipText().ToString(), 0, NodeTemplate->GetKeywords());
 			Action->NodeTemplate = NodeTemplate;
 		}
 	}
diff --git a/Engine/Source/Editor/AnimGraph/Private/AnimationTransitionSchema.cpp b/Engine/Source/Editor/AnimGraph/Private/AnimationTransitionSchema.cpp
--- a/Engine/Source/Editor/AnimGraph/Private/AnimationTransitionSchema.cpp
+++ b/Engine/Source/Editor/AnimGraph/Private/AnimationTransitionSchema.cpp
@@ -94,7 +94,7 @@ void UAnimationTransitionSchema::GetSourceStateActions(FGraphContextMenuBuilder&
 								Args.Add(TEXT("AssetName"), FText::FromString(AssetName));
 								FText Title = FText::Format(LOCTEXT("TransitionFor", "{NodeName} for '{AssetName}'"), Args);
 
-								TSharedPtr<FEdGraphSchemaAction_K2NewNode> Action = FK2ActionMenuBuilder::AddNewNodeAction(ContextMenuBuilder, Category_AssetPlayer, Title, NodeTemplate->GetTooltip(), 0, NodeTemplate->GetKeywords());
+								TSharedPtr<FEdGraphSchemaAction_K2NewNode> Action = FK2ActionMenuBuilder::AddNewNodeAction(ContextMenuBuilder, Category_AssetPlayer, Title, NodeTemplate->GetTooltipText().ToString(), 0, NodeTemplate->GetKeywords());
 								Action->NodeTemplate = NodeTemplate;
 							}
 						}
@@ -117,7 +117,7 @@ void UAnimationTransitionSchema::GetSourceStateActions(FGraphContextMenuBuilder&
 
 						FText Title = UK2Node_TransitionRuleGetter::GetFriendlyName(NodeTemplate->GetterType);
 
-						TSharedPtr<FEdGraphSchemaAction_K2NewNode> Action = FK2ActionMenuBuilder::AddNewNodeAction(ContextMenuBuilder, Category_Transition, Title, NodeTemplate->GetTooltip(), 0, NodeTemplate->GetKeywords());
+						TSharedPtr<FEdGraphSchemaAction_K2NewNode> Action = FK2ActionMenuBuilder::AddNewNodeAction(ContextMenuBuilder, Category_Transition, Title, NodeTemplate->GetTooltipText().ToString(), 0, NodeTemplate->GetKeywords());
 						Action->NodeTemplate = NodeTemplate;
 					}
 				}
diff --git a/Engine/Source/Editor/AnimGraph/Private/K2Node_TransitionRuleGetter.cpp b/Engine/Source/Editor/AnimGraph/Private/K2Node_TransitionRuleGetter.cpp
--- a/Engine/Source/Editor/AnimGraph/Private/K2Node_TransitionRuleGetter.cpp
+++ b/Engine/Source/Editor/AnimGraph/Private/K2Node_TransitionRuleGetter.cpp
@@ -106,9 +106,9 @@ FText UK2Node_TransitionRuleGetter::GetNodeTitle(ENodeTitleType::Type TitleType)
 	return Super::GetNodeTitle(TitleType);
 }
 
-FString UK2Node_TransitionRuleGetter::GetTooltip() const
+FText UK2Node_TransitionRuleGetter::GetTooltipText() const
 {
-	return GetNodeTitle(ENodeTitleType::FullTitle).ToString();
+	return GetNodeTitle(ENodeTitleType::FullTitle);
 }
 
 UEdGraphPin* UK2Node_TransitionRuleGetter::GetOutputPin() const
diff --git a/Engine/Source/Editor/BehaviorTreeEditor/Private/BehaviorTreeGraphNode.cpp b/Engine/Source/Editor/BehaviorTreeEditor/Private/BehaviorTreeGraphNode.cpp
--- a/Engine/Source/Editor/BehaviorTreeEditor/Private/BehaviorTreeGraphNode.cpp
+++ b/Engine/Source/Editor/BehaviorTreeEditor/Private/BehaviorTreeGraphNode.cpp
@@ -127,27 +127,27 @@ FString	UBehaviorTreeGraphNode::GetDescription() const
 		*LOCTEXT("NodeClassError", "Class not found, make sure it's saved!").ToString());
 }
 
-FString UBehaviorTreeGraphNode::GetTooltip() const
+FText UBehaviorTreeGraphNode::GetTooltipText() const
 {
-	FString TooltipDesc;
+	FText TooltipDesc;
 
 	if(!NodeInstance)
 	{
-		TooltipDesc = LOCTEXT("NodeClassError", "Class not found, make sure it's saved!").ToString();
+		TooltipDesc = LOCTEXT("NodeClassError", "Class not found, make sure it's saved!");
 	}
 	else
 	{
 		if(bHasObserverError)
 		{
-			TooltipDesc = LOCTEXT("ObserverError", "Observer has invalid abort setting!").ToString();
+			TooltipDesc = LOCTEXT("ObserverError", "Observer has invalid abort setting!");
 		}
 		else if(DebuggerRuntimeDescription.Len() > 0)
 		{
-			TooltipDesc = DebuggerRuntimeDescription;
+			TooltipDesc = FText::FromString(DebuggerRuntimeDescription);
 		}
 		else if(ErrorMessage.Len() > 0)
 		{
-			TooltipDesc = ErrorMessage;
+			TooltipDesc = FText::FromString(ErrorMessage);
 		}
 		else
 		{
@@ -157,21 +157,23 @@ FString UBehaviorTreeGraphNode::GetTooltip() const
 			{
 				FAssetData AssetData(NodeInstance->GetClass()->ClassGeneratedBy);
 				const FString* Description = AssetData.TagsAndValues.Find(GET_MEMBER_NAME_CHECKED(UBlueprint, BlueprintDescription));
+
 				if (Description && !Description->IsEmpty())
 				{
-					TooltipDesc = Description->Replace(TEXT("\\n"), TEXT("\n"));
+					TooltipDesc = FText::FromString(Description->Replace(TEXT("\\n"), TEXT("\n")));
 				}
 			}
 			else
 			{
-				TooltipDesc = NodeInstance->GetClass()->GetToolTipText().ToString();
+				TooltipDesc = NodeInstance->GetClass()->GetToolTipText();
 			}	
 		}
 	}
 
 	if (bInjectedNode)
 	{
-		TooltipDesc = LOCTEXT("Injected", "Injected: ").ToString() + ((TooltipDesc.Len() > 0) ? TooltipDesc : GetDescription());
+		FText const InjectedDesc = !TooltipDesc.IsEmpty() ? TooltipDesc : FText::FromString(GetDescription());
+		TooltipDesc = FText::Format(LOCTEXT("InjectedTooltip", "Injected: {0}"), InjectedDesc);
 	}
 
 	return TooltipDesc;
diff --git a/Engine/Source/Editor/BehaviorTreeEditor/Private/BehaviorTreeGraphNode_Root.cpp b/Engine/Source/Editor/BehaviorTreeEditor/Private/BehaviorTreeGraphNode_Root.cpp
--- a/Engine/Source/Editor/BehaviorTreeEditor/Private/BehaviorTreeGraphNode_Root.cpp
+++ b/Engine/Source/Editor/BehaviorTreeEditor/Private/BehaviorTreeGraphNode_Root.cpp
@@ -40,9 +40,9 @@ FName UBehaviorTreeGraphNode_Root::GetNameIcon() const
 	return FName("BTEditor.Graph.BTNode.Root.Icon");
 }
 
-FString UBehaviorTreeGraphNode_Root::GetTooltip() const
+FText UBehaviorTreeGraphNode_Root::GetTooltipText() const
 {
-	return UEdGraphNode::GetTooltip();
+	return UEdGraphNode::GetTooltipText();
 }
 
 void UBehaviorTreeGraphNode_Root::PostEditChangeProperty( struct FPropertyChangedEvent& PropertyChangedEvent)
diff --git a/Engine/Source/Editor/BlueprintGraph/Classes/K2Node_Composite.h b/Engine/Source/Editor/BlueprintGraph/Classes/K2Node_Composite.h
--- a/Engine/Source/Editor/BlueprintGraph/Classes/K2Node_Composite.h
+++ b/Engine/Source/Editor/BlueprintGraph/Classes/K2Node_Composite.h
@@ -17,7 +17,7 @@ class UK2Node_Composite : public UK2Node_Tunnel
 	BLUEPRINTGRAPH_API virtual void AllocateDefaultPins() override;
 	BLUEPRINTGRAPH_API virtual void DestroyNode() override;
 	BLUEPRINTGRAPH_API virtual void PostPasteNode() override;
-	BLUEPRINTGRAPH_API virtual FString GetTooltip() const override;
+	BLUEPRINTGRAPH_API virtual FText GetTooltipText() const override;
 	BLUEPRINTGRAPH_API virtual FLinearColor GetNodeTitleColor() const override;
 	BLUEPRINTGRAPH_API virtual FText GetNodeTitle(ENodeTitleType::Type TitleType) const override;
 	BLUEPRINTGRAPH_API virtual bool CanUserDeleteNode() const override;
diff --git a/Engine/Source/Editor/BlueprintGraph/Classes/K2Node_ConstructObjectFromClass.h b/Engine/Source/Editor/BlueprintGraph/Classes/K2Node_ConstructObjectFromClass.h
--- a/Engine/Source/Editor/BlueprintGraph/Classes/K2Node_ConstructObjectFromClass.h
+++ b/Engine/Source/Editor/BlueprintGraph/Classes/K2Node_ConstructObjectFromClass.h
@@ -16,7 +16,7 @@ class BLUEPRINTGRAPH_API UK2Node_ConstructObjectFromClass : public UK2Node
 	virtual FLinearColor GetNodeTitleColor() const override;
 	virtual FText GetNodeTitle(ENodeTitleType::Type TitleType) const override;
 	virtual void PinDefaultValueChanged(UEdGraphPin* Pin) override;
-	virtual FString GetTooltip() const override;
+	virtual FText GetTooltipText() const override;
 	virtual void ExpandNode(class FKismetCompilerContext& CompilerContext, UEdGraph* SourceGraph) override;
 	virtual bool HasExternalBlueprintDependencies(TArray<class UStruct*>* OptionalOutput) const override;
 	virtual bool IsCompatibleWithGraph(const UEdGraph* TargetGraph) const override;
@@ -67,7 +67,7 @@ protected:
 	void SetPinToolTip(UEdGraphPin& MutatablePin, const FText& PinDescription) const;
 
 	/** Tooltip text for this node. */
-	FString NodeTooltip;
+	FText NodeTooltip;
 
 private:
 	/** Constructing FText strings can be costly, so we cache the node's title */
diff --git a/Engine/Source/Editor/BlueprintGraph/Classes/K2Node_EaseFunction.h b/Engine/Source/Editor/BlueprintGraph/Classes/K2Node_EaseFunction.h
--- a/Engine/Source/Editor/BlueprintGraph/Classes/K2Node_EaseFunction.h
+++ b/Engine/Source/Editor/BlueprintGraph/Classes/K2Node_EaseFunction.h
@@ -36,7 +36,7 @@ class UK2Node_EaseFunction : public UK2Node
 	void GetContextMenuActions(const FGraphNodeContextMenuBuilder& Context) const override;
 	virtual void PinDefaultValueChanged(UEdGraphPin* Pin) override;
 	virtual void PinTypeChanged(UEdGraphPin* Pin) override;
-	virtual FString GetTooltip() const override;
+	virtual FText GetTooltipText() const override;
 	virtual FName GetPaletteIcon(FLinearColor& OutColor) const override;
 	// End UEdGraphNode interface.
 
@@ -107,7 +107,7 @@ private:
 	FString EaseFunctionName;
 
 	/** Tooltip text for this node. */
-	FString NodeTooltip;
+	FText NodeTooltip;
 
 	/** The "EaseFunc" input pin, used to enable/disable easing customization */
 	UEdGraphPin* CachedEaseFuncPin;
diff --git a/Engine/Source/Editor/BlueprintGraph/Classes/K2Node_Event.h b/Engine/Source/Editor/BlueprintGraph/Classes/K2Node_Event.h
--- a/Engine/Source/Editor/BlueprintGraph/Classes/K2Node_Event.h
+++ b/Engine/Source/Editor/BlueprintGraph/Classes/K2Node_Event.h
@@ -37,7 +37,7 @@ class UK2Node_Event : public UK2Node_EditablePinBase
 
 	// Begin UEdGraphNode interface
 	BLUEPRINTGRAPH_API virtual void AllocateDefaultPins() override;
-	BLUEPRINTGRAPH_API virtual FString GetTooltip() const override;
+	BLUEPRINTGRAPH_API virtual FText GetTooltipText() const override;
 	BLUEPRINTGRAPH_API virtual FString GetKeywords() const override;	
 	BLUEPRINTGRAPH_API virtual FText GetNodeTitle(ENodeTitleType::Type TitleType) const override;
 	BLUEPRINTGRAPH_API virtual FLinearColor GetNodeTitleColor() const override;
diff --git a/Engine/Source/Editor/BlueprintGraph/Private/K2Node_AIMoveTo.cpp b/Engine/Source/Editor/BlueprintGraph/Private/K2Node_AIMoveTo.cpp
--- a/Engine/Source/Editor/BlueprintGraph/Private/K2Node_AIMoveTo.cpp
+++ b/Engine/Source/Editor/BlueprintGraph/Private/K2Node_AIMoveTo.cpp
@@ -23,9 +23,9 @@ FText UK2Node_AIMoveTo::GetMenuCategory() const
 	return FEditorCategoryUtils::GetCommonCategory(FCommonEditorCategory::AI);
 }
 
-FString UK2Node_AIMoveTo::GetTooltip() const
+FText UK2Node_AIMoveTo::GetTooltipText() const
 {
-	return TEXT("Simple order for Pawn with AIController to move to a specific location");
+	return LOCTEXT("AIMoveTo_Tooltip", "Simple order for Pawn with AIController to move to a specific location");
 }
 
 FText UK2Node_AIMoveTo::GetNodeTitle(ENodeTitleType::Type TitleType) const
diff --git a/Engine/Source/Editor/BlueprintGraph/Private/K2Node_ActorBoundEvent.cpp b/Engine/Source/Editor/BlueprintGraph/Private/K2Node_ActorBoundEvent.cpp
--- a/Engine/Source/Editor/BlueprintGraph/Private/K2Node_ActorBoundEvent.cpp
+++ b/Engine/Source/Editor/BlueprintGraph/Private/K2Node_ActorBoundEvent.cpp
@@ -97,16 +97,16 @@ FText UK2Node_ActorBoundEvent::GetNodeTitle(ENodeTitleType::Type TitleType) cons
 	return CachedNodeTitle;
 }
 
-FString UK2Node_ActorBoundEvent::GetTooltip() const
+FText UK2Node_ActorBoundEvent::GetTooltipText() const
 {
 	UMulticastDelegateProperty* TargetDelegateProp = GetTargetDelegatePropertyConst();
 	if(TargetDelegateProp)
 	{
-		return TargetDelegateProp->GetToolTipText().ToString();
+		return TargetDelegateProp->GetToolTipText();
 	}
 	else
 	{
-		return DelegatePropertyName.ToString();
+		return FText::FromName(DelegatePropertyName);
 	}
 }
 
diff --git a/Engine/Source/Editor/BlueprintGraph/Private/K2Node_AssignmentStatement.cpp b/Engine/Source/Editor/BlueprintGraph/Private/K2Node_AssignmentStatement.cpp
--- a/Engine/Source/Editor/BlueprintGraph/Private/K2Node_AssignmentStatement.cpp
+++ b/Engine/Source/Editor/BlueprintGraph/Private/K2Node_AssignmentStatement.cpp
@@ -75,9 +75,9 @@ void UK2Node_AssignmentStatement::AllocateDefaultPins()
 	Super::AllocateDefaultPins();
 }
 
-FString UK2Node_AssignmentStatement::GetTooltip() const
+FText UK2Node_AssignmentStatement::GetTooltipText() const
 {
-	return TEXT("Assigns Value to Variable");
+	return LOCTEXT("AssignmentStatementTooltip", "Assigns Value to Variable");
 }
 
 FText UK2Node_AssignmentStatement::GetNodeTitle(ENodeTitleType::Type TitleType) const
diff --git a/Engine/Source/Editor/BlueprintGraph/Private/K2Node_BaseAsyncTask.cpp b/Engine/Source/Editor/BlueprintGraph/Private/K2Node_BaseAsyncTask.cpp
--- a/Engine/Source/Editor/BlueprintGraph/Private/K2Node_BaseAsyncTask.cpp
+++ b/Engine/Source/Editor/BlueprintGraph/Private/K2Node_BaseAsyncTask.cpp
@@ -20,10 +20,10 @@ UK2Node_BaseAsyncTask::UK2Node_BaseAsyncTask(const class FPostConstructInitializ
 {
 }
 
-FString UK2Node_BaseAsyncTask::GetTooltip() const
+FText UK2Node_BaseAsyncTask::GetTooltipText() const
 {
 	const FString FunctionToolTipText = UK2Node_CallFunction::GetDefaultTooltipForFunction(GetFactoryFunction());
-	return FunctionToolTipText;
+	return FText::FromString(FunctionToolTipText);
 }
 
 FText UK2Node_BaseAsyncTask::GetNodeTitle(ENodeTitleType::Type TitleType) const
@@ -58,7 +58,7 @@ void UK2Node_BaseAsyncTask::GetMenuEntries(FGraphContextMenuBuilder& ContextMenu
 
 TSharedPtr<FEdGraphSchemaAction_K2NewNode> UK2Node_BaseAsyncTask::CreateDefaultMenuEntry(UK2Node_BaseAsyncTask* NodeTemplate, FGraphContextMenuBuilder& ContextMenuBuilder) const
 {
-	TSharedPtr<FEdGraphSchemaAction_K2NewNode> NodeAction = FK2ActionMenuBuilder::AddNewNodeAction(ContextMenuBuilder, NodeTemplate->GetMenuCategory().ToString(), NodeTemplate->GetNodeTitle(ENodeTitleType::ListView), NodeTemplate->GetTooltip(), 0, NodeTemplate->GetKeywords());
+	TSharedPtr<FEdGraphSchemaAction_K2NewNode> NodeAction = FK2ActionMenuBuilder::AddNewNodeAction(ContextMenuBuilder, NodeTemplate->GetMenuCategory().ToString(), NodeTemplate->GetNodeTitle(ENodeTitleType::ListView), NodeTemplate->GetTooltipText().ToString(), 0, NodeTemplate->GetKeywords());
 	
 	NodeAction->NodeTemplate = NodeTemplate;
 
diff --git a/Engine/Source/Editor/BlueprintGraph/Private/K2Node_BreakStruct.cpp b/Engine/Source/Editor/BlueprintGraph/Private/K2Node_BreakStruct.cpp
--- a/Engine/Source/Editor/BlueprintGraph/Private/K2Node_BreakStruct.cpp
+++ b/Engine/Source/Editor/BlueprintGraph/Private/K2Node_BreakStruct.cpp
@@ -205,11 +205,11 @@ FText UK2Node_BreakStruct::GetNodeTitle(ENodeTitleType::Type TitleType) const
 	return CachedNodeTitle;
 }
 
-FString UK2Node_BreakStruct::GetTooltip() const
+FText UK2Node_BreakStruct::GetTooltipText() const
 {
-	return FString::Printf(
-		*LOCTEXT("MakeStruct_Tooltip", "Adds a node that breaks a '%s' into its member fields").ToString(),
-		*(StructType ? StructType->GetName() : FString())
+	return FText::Format(
+		LOCTEXT("MakeStruct_Tooltip", "Adds a node that breaks a '{0}' into its member fields"),
+		StructType ? FText::FromName(StructType->GetFName()) : FText::GetEmpty()
 		);
 }
 
diff --git a/Engine/Source/Editor/BlueprintGraph/Private/K2Node_CallFunction.cpp b/Engine/Source/Editor/BlueprintGraph/Private/K2Node_CallFunction.cpp
--- a/Engine/Source/Editor/BlueprintGraph/Private/K2Node_CallFunction.cpp
+++ b/Engine/Source/Editor/BlueprintGraph/Private/K2Node_CallFunction.cpp
@@ -753,30 +753,37 @@ FLinearColor UK2Node_CallFunction::GetNodeTitleColor() const
 		return Super::GetNodeTitleColor();
 	}
 
-FString UK2Node_CallFunction::GetTooltip() const
+FText UK2Node_CallFunction::GetTooltipText() const
 {
-	FString Tooltip;
+	FText Tooltip;
 
 	if (UFunction* Function = GetTargetFunction())
 	{
-		Tooltip = GetDefaultTooltipForFunction(Function);
+		Tooltip = FText::FromString(GetDefaultTooltipForFunction(Function));
 
-		FString ClientString;
+		FFormatNamedArguments Args;
+		Args.Add(TEXT("DefaultTooltip"), Tooltip);
 
 		if (Function->HasAllFunctionFlags(FUNC_BlueprintAuthorityOnly))
 		{
-			ClientString = NSLOCTEXT("K2Node", "ServerFunction", "\n\nAuthority Only. This function will only execute on the server.").ToString();
+			Args.Add(
+				TEXT("ClientString"),
+				NSLOCTEXT("K2Node", "ServerFunction", "Authority Only. This function will only execute on the server.")
+			);
+			Tooltip = FText::Format(LOCTEXT("CallFunction_SubtitledTooltip", "{DefaultTooltip}\n\n{ClientString}"), Args);
 		}
 		else if (Function->HasAllFunctionFlags(FUNC_BlueprintCosmetic))
 		{
-			ClientString = NSLOCTEXT("K2Node", "ClientEvent", "\n\nCosmetic. This event is only for cosmetic, non-gameplay actions.").ToString();
+			Args.Add(
+				TEXT("ClientString"),
+				NSLOCTEXT("K2Node", "ClientEvent", "Cosmetic. This event is only for cosmetic, non-gameplay actions.")
+			);
+			Tooltip = FText::Format(LOCTEXT("CallFunction_SubtitledTooltip", "{DefaultTooltip}\n\n{ClientString}"), Args);
 		} 
-
-		Tooltip += ClientString;
 	}
 	else
 	{
-		Tooltip = LOCTEXT("CallUnknownFunction", "Call unknown function ").ToString() + FunctionReference.GetMemberName().ToString();
+		Tooltip = FText::Format(LOCTEXT("CallUnknownFunction", "Call unknown function {0}"), FText::FromName(FunctionReference.GetMemberName()));
 	}
 
 	return Tooltip;
diff --git a/Engine/Source/Editor/BlueprintGraph/Private/K2Node_CastByteToEnum.cpp b/Engine/Source/Editor/BlueprintGraph/Private/K2Node_CastByteToEnum.cpp
--- a/Engine/Source/Editor/BlueprintGraph/Private/K2Node_CastByteToEnum.cpp
+++ b/Engine/Source/Editor/BlueprintGraph/Private/K2Node_CastByteToEnum.cpp
@@ -26,16 +26,17 @@ void UK2Node_CastByteToEnum::AllocateDefaultPins()
 	CreatePin(EGPD_Output, Schema->PC_Byte, TEXT(""), Enum, false, false, Schema->PN_ReturnValue);
 }
 
-FString UK2Node_CastByteToEnum::GetTooltip() const
+FText UK2Node_CastByteToEnum::GetTooltipText() const
 {
-	return FString::Printf( 
-		*NSLOCTEXT("K2Node", "CastByteToEnum_Tooltip", "Byte to Enum %s").ToString(),
-		*Enum->GetName());
+	return FText::Format(
+		NSLOCTEXT("K2Node", "CastByteToEnum_Tooltip", "Byte to Enum {0}"),
+		FText::FromName(Enum->GetFName())
+	);
 }
 
 FText UK2Node_CastByteToEnum::GetNodeTitle(ENodeTitleType::Type TitleType) const
 {
-	return FText::FromString(GetTooltip());
+	return GetTooltipText();
 }
 
 FText UK2Node_CastByteToEnum::GetCompactNodeTitle() const
diff --git a/Engine/Source/Editor/BlueprintGraph/Private/K2Node_ComponentBoundEvent.cpp b/Engine/Source/Editor/BlueprintGraph/Private/K2Node_ComponentBoundEvent.cpp
--- a/Engine/Source/Editor/BlueprintGraph/Private/K2Node_ComponentBoundEvent.cpp
+++ b/Engine/Source/Editor/BlueprintGraph/Private/K2Node_ComponentBoundEvent.cpp
@@ -67,16 +67,16 @@ UMulticastDelegateProperty* UK2Node_ComponentBoundEvent::GetTargetDelegateProper
 }
 
 
-FString UK2Node_ComponentBoundEvent::GetTooltip() const
+FText UK2Node_ComponentBoundEvent::GetTooltipText() const
 {
 	UMulticastDelegateProperty* TargetDelegateProp = GetTargetDelegateProperty();
 	if (TargetDelegateProp)
 	{
-		return TargetDelegateProp->GetToolTipText().ToString();
+		return TargetDelegateProp->GetToolTipText();
 	}
 	else
 	{
-		return DelegatePropertyName.ToString();
+		return FText::FromName(DelegatePropertyName);
 	}
 }
 
diff --git a/Engine/Source/Editor/BlueprintGraph/Private/K2Node_Composite.cpp b/Engine/Source/Editor/BlueprintGraph/Private/K2Node_Composite.cpp
--- a/Engine/Source/Editor/BlueprintGraph/Private/K2Node_Composite.cpp
+++ b/Engine/Source/Editor/BlueprintGraph/Private/K2Node_Composite.cpp
@@ -139,17 +139,17 @@ void UK2Node_Composite::PostPasteNode()
 	}
 }
 
-FString UK2Node_Composite::GetTooltip() const
+FText UK2Node_Composite::GetTooltipText() const
 {
 	if (InputSinkNode != NULL)
 	{
 		if (!InputSinkNode->MetaData.ToolTip.IsEmpty())
 		{
-			return InputSinkNode->MetaData.ToolTip;
+			return FText::FromString(InputSinkNode->MetaData.ToolTip);
 		}
 	}
 
-	return FString::Printf(*LOCTEXT("CollapsedCompositeNode", "Collapsed composite node").ToString());
+	return LOCTEXT("CollapsedCompositeNode", "Collapsed composite node");
 }
 
 FLinearColor UK2Node_Composite::GetNodeTitleColor() const
diff --git a/Engine/Source/Editor/BlueprintGraph/Private/K2Node_ConstructObjectFromClass.cpp b/Engine/Source/Editor/BlueprintGraph/Private/K2Node_ConstructObjectFromClass.cpp
--- a/Engine/Source/Editor/BlueprintGraph/Private/K2Node_ConstructObjectFromClass.cpp
+++ b/Engine/Source/Editor/BlueprintGraph/Private/K2Node_ConstructObjectFromClass.cpp
@@ -21,7 +21,7 @@ FString FK2Node_ConstructObjectFromClassHelper::ClassPinName(TEXT("Class"));
 UK2Node_ConstructObjectFromClass::UK2Node_ConstructObjectFromClass(const class FPostConstructInitializeProperties& PCIP)
 	: Super(PCIP)
 {
-	NodeTooltip = LOCTEXT("NodeTooltip", "Attempts to spawn a new object").ToString();
+	NodeTooltip = LOCTEXT("NodeTooltip", "Attempts to spawn a new object");
 }
 
 UClass* UK2Node_ConstructObjectFromClass::GetClassPinBaseClass() const
@@ -181,7 +181,7 @@ void UK2Node_ConstructObjectFromClass::PinDefaultValueChanged(UEdGraphPin* Chang
 	}
 }
 
-FString UK2Node_ConstructObjectFromClass::GetTooltip() const
+FText UK2Node_ConstructObjectFromClass::GetTooltipText() const
 {
 	return NodeTooltip;
 }
diff --git a/Engine/Source/Editor/BlueprintGraph/Private/K2Node_CustomEvent.cpp b/Engine/Source/Editor/BlueprintGraph/Private/K2Node_CustomEvent.cpp
--- a/Engine/Source/Editor/BlueprintGraph/Private/K2Node_CustomEvent.cpp
+++ b/Engine/Source/Editor/BlueprintGraph/Private/K2Node_CustomEvent.cpp
@@ -358,9 +358,9 @@ bool UK2Node_CustomEvent::IsUsedByAuthorityOnlyDelegate() const
 	return false;
 }
 
-FString UK2Node_CustomEvent::GetTooltip() const
+FText UK2Node_CustomEvent::GetTooltipText() const
 {
-	return LOCTEXT("AddCustomEvent_Tooltip", "An event with customizable name and parameters.").ToString();
+	return LOCTEXT("AddCustomEvent_Tooltip", "An event with customizable name and parameters.");
 }
 
 FString UK2Node_CustomEvent::GetDocumentationLink() const
diff --git a/Engine/Source/Editor/BlueprintGraph/Private/K2Node_DelegateSet.cpp b/Engine/Source/Editor/BlueprintGraph/Private/K2Node_DelegateSet.cpp
--- a/Engine/Source/Editor/BlueprintGraph/Private/K2Node_DelegateSet.cpp
+++ b/Engine/Source/Editor/BlueprintGraph/Private/K2Node_DelegateSet.cpp
@@ -147,16 +147,16 @@ void UK2Node_DelegateSet::AllocateDefaultPins()
 	Super::AllocateDefaultPins();
 }
 
-FString UK2Node_DelegateSet::GetTooltip() const
+FText UK2Node_DelegateSet::GetTooltipText() const
 {
-	FString Tooltip = FString::Printf(*NSLOCTEXT("K2Node", "CreateEventForDelegate", "Create an event tied to the delegate %s").ToString(), *DelegatePropertyName.ToString());
+	FText Tooltip = FText::Format(NSLOCTEXT("K2Node", "CreateEventForDelegate", "Create an event tied to the delegate {0}"), FText::FromName(DelegatePropertyName));
 	if (UFunction* Function = GetDelegateSignature())
 	{
-		const FString SignatureTooltip = Function->GetToolTipText().ToString();
+		const FText SignatureTooltip = Function->GetToolTipText();
 
 		if (!SignatureTooltip.IsEmpty())
 		{
-			Tooltip += FString(TEXT("\n")) + SignatureTooltip;
+			Tooltip = FText::Format(LOCTEXT("DelegateSet_SubtitledTooltip", "{0}\n{1}"), Tooltip, SignatureTooltip);
 		}
 	}
 	return Tooltip;
diff --git a/Engine/Source/Editor/BlueprintGraph/Private/K2Node_DoOnceMultiInput.cpp b/Engine/Source/Editor/BlueprintGraph/Private/K2Node_DoOnceMultiInput.cpp
--- a/Engine/Source/Editor/BlueprintGraph/Private/K2Node_DoOnceMultiInput.cpp
+++ b/Engine/Source/Editor/BlueprintGraph/Private/K2Node_DoOnceMultiInput.cpp
@@ -20,7 +20,7 @@ void UK2Node_DoOnceMultiInput::GetMenuEntries(FGraphContextMenuBuilder& ContextM
 	// Create a node template for this factory method
 	UK2Node_DoOnceMultiInput* NodeTemplate = NewObject<UK2Node_DoOnceMultiInput>(ContextMenuBuilder.OwnerOfTemporaries);
 	
-	TSharedPtr<FEdGraphSchemaAction_K2NewNode> NodeAction = FK2ActionMenuBuilder::AddNewNodeAction(ContextMenuBuilder, LOCTEXT("FlowCategory", "Utilities|Flow Control").ToString(), NodeTemplate->GetNodeTitle(ENodeTitleType::ListView), NodeTemplate->GetTooltip(), 0, NodeTemplate->GetKeywords());
+	TSharedPtr<FEdGraphSchemaAction_K2NewNode> NodeAction = FK2ActionMenuBuilder::AddNewNodeAction(ContextMenuBuilder, LOCTEXT("FlowCategory", "Utilities|Flow Control").ToString(), NodeTemplate->GetNodeTitle(ENodeTitleType::ListView), NodeTemplate->GetTooltipText().ToString(), 0, NodeTemplate->GetKeywords());
 
 	NodeAction->NodeTemplate = NodeTemplate;
 }
diff --git a/Engine/Source/Editor/BlueprintGraph/Private/K2Node_EaseFunction.cpp b/Engine/Source/Editor/BlueprintGraph/Private/K2Node_EaseFunction.cpp
--- a/Engine/Source/Editor/BlueprintGraph/Private/K2Node_EaseFunction.cpp
+++ b/Engine/Source/Editor/BlueprintGraph/Private/K2Node_EaseFunction.cpp
@@ -70,7 +70,7 @@ UK2Node_EaseFunction::UK2Node_EaseFunction(const class FPostConstructInitializeP
 	: Super(PCIP)
 	, CachedEaseFuncPin(NULL)
 {
-	NodeTooltip = LOCTEXT("NodeTooltip", "Interpolates from value A to value B using a user specified easing function").ToString();
+	NodeTooltip = LOCTEXT("NodeTooltip", "Interpolates from value A to value B using a user specified easing function");
 	OldEasingFunc = INDEX_NONE;
 	EaseFunctionName = TEXT("");
 }
@@ -112,7 +112,7 @@ FText UK2Node_EaseFunction::GetNodeTitle(ENodeTitleType::Type TitleType) const
 	return LOCTEXT("EaseFunction_Title", "Ease");
 }
 
-FString UK2Node_EaseFunction::GetTooltip() const
+FText UK2Node_EaseFunction::GetTooltipText() const
 {
 	return NodeTooltip;
 }
diff --git a/Engine/Source/Editor/BlueprintGraph/Private/K2Node_EnumEquality.cpp b/Engine/Source/Editor/BlueprintGraph/Private/K2Node_EnumEquality.cpp
--- a/Engine/Source/Editor/BlueprintGraph/Private/K2Node_EnumEquality.cpp
+++ b/Engine/Source/Editor/BlueprintGraph/Private/K2Node_EnumEquality.cpp
@@ -31,9 +31,9 @@ void UK2Node_EnumEquality::AllocateDefaultPins()
 	Super::AllocateDefaultPins();
 }
 
-FString UK2Node_EnumEquality::GetTooltip() const
+FText UK2Node_EnumEquality::GetTooltipText() const
 {
-	return TEXT("Returns true if A is equal to B (A == B)");
+	return LOCTEXT("EnumEqualityTooltip", "Returns true if A is equal to B (A == B)");
 }
 
 FText UK2Node_EnumEquality::GetNodeTitle(ENodeTitleType::Type TitleType) const
@@ -208,7 +208,7 @@ void UK2Node_EnumEquality::GetMenuEntries(FGraphContextMenuBuilder& Context) con
 
 		const FString Category = TEXT("Utilities| Enum");
 		const FText MenuDesc = EnumNodeTemplate->GetNodeTitle(ENodeTitleType::ListView);
-		const FString Tooltip = EnumNodeTemplate->GetTooltip();
+		const FString Tooltip = EnumNodeTemplate->GetTooltipText().ToString();
 		const FString Keywords = EnumNodeTemplate->GetKeywords();
 
 		TSharedPtr<FEdGraphSchemaAction_K2NewNode> NodeAction = FK2ActionMenuBuilder::AddNewNodeAction(Context, Category, MenuDesc, Tooltip, 0, Keywords);
diff --git a/Engine/Source/Editor/BlueprintGraph/Private/K2Node_EnumLiteral.cpp b/Engine/Source/Editor/BlueprintGraph/Private/K2Node_EnumLiteral.cpp
--- a/Engine/Source/Editor/BlueprintGraph/Private/K2Node_EnumLiteral.cpp
+++ b/Engine/Source/Editor/BlueprintGraph/Private/K2Node_EnumLiteral.cpp
@@ -39,11 +39,11 @@ FName UK2Node_EnumLiteral::GetPaletteIcon(FLinearColor& OutColor) const
 	return PaletteIconName;
 }
 
-FString UK2Node_EnumLiteral::GetTooltip() const
+FText UK2Node_EnumLiteral::GetTooltipText() const
 {
 	if (Enum == nullptr)
 	{
-		return NSLOCTEXT("K2Node", "BadEnumLiteral_Tooltip", "Literal enum (bad enum)").ToString();
+		return NSLOCTEXT("K2Node", "BadEnumLiteral_Tooltip", "Literal enum (bad enum)");
 	}
 	else if (CachedTooltip.IsOutOfDate())
 	{
@@ -51,14 +51,14 @@ FString UK2Node_EnumLiteral::GetTooltip() const
 		Args.Add(TEXT("EnumName"), FText::FromName(Enum->GetFName()));
 		CachedTooltip = FText::Format(NSLOCTEXT("K2Node", "EnumLiteral_Tooltip", "Literal enum {EnumName}"), Args);
 	}
-	return ((FText&)CachedTooltip).ToString();	
+	return CachedTooltip;	
 }
 
 FText UK2Node_EnumLiteral::GetNodeTitle(ENodeTitleType::Type TitleType) const
 {
 	if (CachedTooltip.IsOutOfDate())
 	{
-		return FText::FromString(GetTooltip());
+		return GetTooltipText();
 	}
 	return CachedTooltip;
 }
diff --git a/Engine/Source/Editor/BlueprintGraph/Private/K2Node_Event.cpp b/Engine/Source/Editor/BlueprintGraph/Private/K2Node_Event.cpp
--- a/Engine/Source/Editor/BlueprintGraph/Private/K2Node_Event.cpp
+++ b/Engine/Source/Editor/BlueprintGraph/Private/K2Node_Event.cpp
@@ -89,31 +89,38 @@ FText UK2Node_Event::GetNodeTitle(ENodeTitleType::Type TitleType) const
 	}
 }
 
-FString UK2Node_Event::GetTooltip() const
+FText UK2Node_Event::GetTooltipText() const
 {
-	FString Tooltip;
+	FText Tooltip;
 
 	UFunction* Function = FindField<UFunction>(EventSignatureClass, EventSignatureName);
 	if (Function != NULL)
 	{
-		Tooltip = Function->GetToolTipText().ToString();
+		Tooltip = Function->GetToolTipText();
 
 		if (bOverrideFunction || (CustomFunctionName == NAME_None))
 		{
-			FString ClientString;
+			FFormatNamedArguments Args;
+			Args.Add(TEXT("FunctionTooltip"), Tooltip);
 
 			//@TODO: KISMETREPLICATION: Should do this for events with a custom function name, if it's a newly introduced replicating thingy
 			if (Function->HasAllFunctionFlags(FUNC_BlueprintCosmetic) || IsCosmeticTickEvent())
 			{
-				ClientString = NSLOCTEXT("K2Node", "ClientEvent", "\n\nCosmetic. This event is only for cosmetic, non-gameplay actions.").ToString();
+				Args.Add(
+					TEXT("ClientString"),
+					NSLOCTEXT("K2Node", "ClientEvent", "\n\nCosmetic. This event is only for cosmetic, non-gameplay actions.")
+				);
+				Tooltip = FText::Format(LOCTEXT("Event_SubtitledTooltip", "{FunctionTooltip}\n\n{ClientString}"), Args);
 			} 
 			else if(Function->HasAllFunctionFlags(FUNC_BlueprintAuthorityOnly))
 			{
-				ClientString = NSLOCTEXT("K2Node", "ServerEvent", "\n\nAuthority Only. This event only fires on the server.").ToString();
-			}
-
-			Tooltip += ClientString;
-		}
+				Args.Add(
+					TEXT("ClientString"),
+					NSLOCTEXT("K2Node", "ServerEvent", "\n\nAuthority Only. This event only fires on the server.")
+				);
+				Tooltip = FText::Format(LOCTEXT("Event_SubtitledTooltip", "{FunctionTooltip}\n\n{ClientString}"), Args);
+			}			
+		}		
 	}
 
 	return Tooltip;
diff --git a/Engine/Source/Editor/BlueprintGraph/Private/K2Node_ExecutionSequence.cpp b/Engine/Source/Editor/BlueprintGraph/Private/K2Node_ExecutionSequence.cpp
--- a/Engine/Source/Editor/BlueprintGraph/Private/K2Node_ExecutionSequence.cpp
+++ b/Engine/Source/Editor/BlueprintGraph/Private/K2Node_ExecutionSequence.cpp
@@ -150,9 +150,9 @@ FLinearColor UK2Node_ExecutionSequence::GetNodeTitleColor() const
 	return FLinearColor::White;
 }
 
-FString UK2Node_ExecutionSequence::GetTooltip() const
+FText UK2Node_ExecutionSequence::GetTooltipText() const
 {
-	return NSLOCTEXT("K2Node", "ExecutePinInOrder_Tooltip", "Executes a series of pins in order").ToString();
+	return NSLOCTEXT("K2Node", "ExecutePinInOrder_Tooltip", "Executes a series of pins in order");
 }
 
 FString UK2Node_ExecutionSequence::GetUniquePinName()
diff --git a/Engine/Source/Editor/BlueprintGraph/Private/K2Node_ForEachElementInEnum.cpp b/Engine/Source/Editor/BlueprintGraph/Private/K2Node_ForEachElementInEnum.cpp
--- a/Engine/Source/Editor/BlueprintGraph/Private/K2Node_ForEachElementInEnum.cpp
+++ b/Engine/Source/Editor/BlueprintGraph/Private/K2Node_ForEachElementInEnum.cpp
@@ -131,9 +131,9 @@ void UK2Node_ForEachElementInEnum::ValidateNodeDuringCompilation(FCompilerResult
 	}
 }
 
-FString UK2Node_ForEachElementInEnum::GetTooltip() const
+FText UK2Node_ForEachElementInEnum::GetTooltipText() const
 {
-	return GetNodeTitle(ENodeTitleType::FullTitle).ToString();
+	return GetNodeTitle(ENodeTitleType::FullTitle);
 }
 
 FText UK2Node_ForEachElementInEnum::GetNodeTitle(ENodeTitleType::Type TitleType) const
diff --git a/Engine/Source/Editor/BlueprintGraph/Private/K2Node_FormatText.cpp b/Engine/Source/Editor/BlueprintGraph/Private/K2Node_FormatText.cpp
--- a/Engine/Source/Editor/BlueprintGraph/Private/K2Node_FormatText.cpp
+++ b/Engine/Source/Editor/BlueprintGraph/Private/K2Node_FormatText.cpp
@@ -28,7 +28,7 @@ UK2Node_FormatText::UK2Node_FormatText(const class FPostConstructInitializePrope
 	: Super(PCIP)
 	, CachedFormatPin(NULL)
 {
-	NodeTooltip = LOCTEXT("NodeTooltip", "Builds a formatted string using available specifier values. Use {} to denote specifiers.").ToString();
+	NodeTooltip = LOCTEXT("NodeTooltip", "Builds a formatted string using available specifier values. Use {} to denote specifiers.");
 }
 
 void UK2Node_FormatText::AllocateDefaultPins()
@@ -156,7 +156,7 @@ void UK2Node_FormatText::PinDefaultValueChanged(UEdGraphPin* Pin)
 	}
 }
 
-FString UK2Node_FormatText::GetTooltip() const
+FText UK2Node_FormatText::GetTooltipText() const
 {
 	return NodeTooltip;
 }
diff --git a/Engine/Source/Editor/BlueprintGraph/Private/K2Node_GetDataTableRow.cpp b/Engine/Source/Editor/BlueprintGraph/Private/K2Node_GetDataTableRow.cpp
--- a/Engine/Source/Editor/BlueprintGraph/Private/K2Node_GetDataTableRow.cpp
+++ b/Engine/Source/Editor/BlueprintGraph/Private/K2Node_GetDataTableRow.cpp
@@ -24,7 +24,7 @@ FString UK2Node_GetDataTableRowHelper::RowNamePinName(LOCTEXT("RowNamePinName","
 UK2Node_GetDataTableRow::UK2Node_GetDataTableRow(const class FPostConstructInitializeProperties& PCIP)
 	: Super(PCIP)
 {
-	NodeTooltip = LOCTEXT("NodeTooltip", "Attempts to retrieve a TableRow from a DataTable via it's RowName").ToString();
+	NodeTooltip = LOCTEXT("NodeTooltip", "Attempts to retrieve a TableRow from a DataTable via it's RowName");
 }
 
 void UK2Node_GetDataTableRow::AllocateDefaultPins()
@@ -209,7 +209,7 @@ void UK2Node_GetDataTableRow::PinDefaultValueChanged(UEdGraphPin* ChangedPin)
 	}
 }
 
-FString UK2Node_GetDataTableRow::GetTooltip() const
+FText UK2Node_GetDataTableRow::GetTooltipText() const
 {
 	return NodeTooltip;
 }
diff --git a/Engine/Source/Editor/BlueprintGraph/Private/K2Node_GetEnumeratorName.cpp b/Engine/Source/Editor/BlueprintGraph/Private/K2Node_GetEnumeratorName.cpp
--- a/Engine/Source/Editor/BlueprintGraph/Private/K2Node_GetEnumeratorName.cpp
+++ b/Engine/Source/Editor/BlueprintGraph/Private/K2Node_GetEnumeratorName.cpp
@@ -25,9 +25,9 @@ void UK2Node_GetEnumeratorName::AllocateDefaultPins()
 	CreatePin(EGPD_Output, Schema->PC_Name, TEXT(""), NULL, false, false, Schema->PN_ReturnValue);
 }
 
-FString UK2Node_GetEnumeratorName::GetTooltip() const
+FText UK2Node_GetEnumeratorName::GetTooltipText() const
 {
-	return NSLOCTEXT("K2Node", "GetEnumeratorName_Tooltip", "Returns name of enumerator").ToString();
+	return NSLOCTEXT("K2Node", "GetEnumeratorName_Tooltip", "Returns name of enumerator");
 }
 
 FText UK2Node_GetEnumeratorName::GetNodeTitle(ENodeTitleType::Type TitleType) const
diff --git a/Engine/Source/Editor/BlueprintGraph/Private/K2Node_GetEnumeratorNameAsString.cpp b/Engine/Source/Editor/BlueprintGraph/Private/K2Node_GetEnumeratorNameAsString.cpp
--- a/Engine/Source/Editor/BlueprintGraph/Private/K2Node_GetEnumeratorNameAsString.cpp
+++ b/Engine/Source/Editor/BlueprintGraph/Private/K2Node_GetEnumeratorNameAsString.cpp
@@ -22,9 +22,9 @@ void UK2Node_GetEnumeratorNameAsString::AllocateDefaultPins()
 	CreatePin(EGPD_Output, Schema->PC_String, TEXT(""), NULL, false, false, Schema->PN_ReturnValue);
 }
 
-FString UK2Node_GetEnumeratorNameAsString::GetTooltip() const
+FText UK2Node_GetEnumeratorNameAsString::GetTooltipText() const
 {
-	return NSLOCTEXT("K2Node", "GetEnumeratorNameAsString_Tooltip", "Returns user friendly name of enumerator").ToString();
+	return NSLOCTEXT("K2Node", "GetEnumeratorNameAsString_Tooltip", "Returns user friendly name of enumerator");
 }
 
 FText UK2Node_GetEnumeratorNameAsString::GetNodeTitle(ENodeTitleType::Type TitleType) const
diff --git a/Engine/Source/Editor/BlueprintGraph/Private/K2Node_GetInputAxisKeyValue.cpp b/Engine/Source/Editor/BlueprintGraph/Private/K2Node_GetInputAxisKeyValue.cpp
--- a/Engine/Source/Editor/BlueprintGraph/Private/K2Node_GetInputAxisKeyValue.cpp
+++ b/Engine/Source/Editor/BlueprintGraph/Private/K2Node_GetInputAxisKeyValue.cpp
@@ -50,11 +50,11 @@ FText UK2Node_GetInputAxisKeyValue::GetNodeTitle(ENodeTitleType::Type TitleType)
 	return CachedNodeTitle;
 }
 
-FString UK2Node_GetInputAxisKeyValue::GetTooltip() const
+FText UK2Node_GetInputAxisKeyValue::GetTooltipText() const
 {
 	FFormatNamedArguments Args;
 	Args.Add(TEXT("AxisKey"), InputAxisKey.GetDisplayName());
-	return FText::Format(NSLOCTEXT("K2Node", "GetInputAxisKey_Tooltip", "Returns the current value of input axis key {AxisKey}.  If input is disabled for the actor the value will be 0."), Args).ToString();
+	return FText::Format(NSLOCTEXT("K2Node", "GetInputAxisKey_Tooltip", "Returns the current value of input axis key {AxisKey}.  If input is disabled for the actor the value will be 0."), Args);
 }
 
 bool UK2Node_GetInputAxisKeyValue::IsCompatibleWithGraph(UEdGraph const* Graph) const
diff --git a/Engine/Source/Editor/BlueprintGraph/Private/K2Node_GetInputAxisValue.cpp b/Engine/Source/Editor/BlueprintGraph/Private/K2Node_GetInputAxisValue.cpp
--- a/Engine/Source/Editor/BlueprintGraph/Private/K2Node_GetInputAxisValue.cpp
+++ b/Engine/Source/Editor/BlueprintGraph/Private/K2Node_GetInputAxisValue.cpp
@@ -51,9 +51,9 @@ FText UK2Node_GetInputAxisValue::GetNodeTitle(ENodeTitleType::Type TitleType) co
 	return CachedNodeTitle;
 }
 
-FString UK2Node_GetInputAxisValue::GetTooltip() const
+FText UK2Node_GetInputAxisValue::GetTooltipText() const
 {
-	return FString::Printf(*NSLOCTEXT("K2Node", "GetInputAxis_Tooltip", "Returns the current value of input axis %s.  If input is disabled for the actor the value will be 0.").ToString(), *InputAxisName.ToString());
+	return FText::Format(NSLOCTEXT("K2Node", "GetInputAxis_Tooltip", "Returns the current value of input axis {0}.  If input is disabled for the actor the value will be 0."), FText::FromName(InputAxisName));
 }
 
 bool UK2Node_GetInputAxisValue::IsCompatibleWithGraph(UEdGraph const* Graph) const
diff --git a/Engine/Source/Editor/BlueprintGraph/Private/K2Node_GetInputVectorAxisValue.cpp b/Engine/Source/Editor/BlueprintGraph/Private/K2Node_GetInputVectorAxisValue.cpp
--- a/Engine/Source/Editor/BlueprintGraph/Private/K2Node_GetInputVectorAxisValue.cpp
+++ b/Engine/Source/Editor/BlueprintGraph/Private/K2Node_GetInputVectorAxisValue.cpp
@@ -19,11 +19,11 @@ void UK2Node_GetInputVectorAxisValue::Initialize(const FKey AxisKey)
 	SetFromFunction(AActor::StaticClass()->FindFunctionByName(TEXT("GetInputVectorAxisValue")));
 }
 
-FString UK2Node_GetInputVectorAxisValue::GetTooltip() const
+FText UK2Node_GetInputVectorAxisValue::GetTooltipText() const
 {
 	FFormatNamedArguments Args;
 	Args.Add(TEXT("AxisKey"), InputAxisKey.GetDisplayName());
-	return FText::Format(NSLOCTEXT("K2Node", "GetInputVectorAxis_Tooltip", "Returns the current value of input axis key {AxisKey}.  If input is disabled for the actor the value will be (0, 0, 0)."), Args).ToString();
+	return FText::Format(NSLOCTEXT("K2Node", "GetInputVectorAxis_Tooltip", "Returns the current value of input axis key {AxisKey}.  If input is disabled for the actor the value will be (0, 0, 0)."), Args);
 }
 
 void UK2Node_GetInputVectorAxisValue::ValidateNodeDuringCompilation(class FCompilerResultsLog& MessageLog) const
diff --git a/Engine/Source/Editor/BlueprintGraph/Private/K2Node_GetNumEnumEntries.cpp b/Engine/Source/Editor/BlueprintGraph/Private/K2Node_GetNumEnumEntries.cpp
--- a/Engine/Source/Editor/BlueprintGraph/Private/K2Node_GetNumEnumEntries.cpp
+++ b/Engine/Source/Editor/BlueprintGraph/Private/K2Node_GetNumEnumEntries.cpp
@@ -24,11 +24,11 @@ void UK2Node_GetNumEnumEntries::AllocateDefaultPins()
 	Super::AllocateDefaultPins();
 }
 
-FString UK2Node_GetNumEnumEntries::GetTooltip() const
+FText UK2Node_GetNumEnumEntries::GetTooltipText() const
 {
 	const FString EnumName = (Enum != NULL) ? Enum->GetName() : TEXT("(bad enum)");
 
-	return FString::Printf(*NSLOCTEXT("K2Node", "GetNumEnumEntries_Tooltip", "Returns %s_MAX value").ToString(), *EnumName);
+	return FText::Format(NSLOCTEXT("K2Node", "GetNumEnumEntries_Tooltip", "Returns {0}_MAX value"), FText::FromString(EnumName));
 }
 
 FText UK2Node_GetNumEnumEntries::GetNodeTitle(ENodeTitleType::Type TitleType) const
diff --git a/Engine/Source/Editor/BlueprintGraph/Private/K2Node_IfThenElse.cpp b/Engine/Source/Editor/BlueprintGraph/Private/K2Node_IfThenElse.cpp
--- a/Engine/Source/Editor/BlueprintGraph/Private/K2Node_IfThenElse.cpp
+++ b/Engine/Source/Editor/BlueprintGraph/Private/K2Node_IfThenElse.cpp
@@ -104,9 +104,9 @@ FLinearColor UK2Node_IfThenElse::GetNodeTitleColor() const
 	return GetDefault<UGraphEditorSettings>()->ExecBranchNodeTitleColor;
 }
 
-FString UK2Node_IfThenElse::GetTooltip() const
+FText UK2Node_IfThenElse::GetTooltipText() const
 {
-	return *LOCTEXT("BrancStatement_Tooltip", "Branch Statement\nIf Condition is true, execution goes to True, otherwise it goes to False").ToString();
+	return LOCTEXT("BrancStatement_Tooltip", "Branch Statement\nIf Condition is true, execution goes to True, otherwise it goes to False");
 }
 
 FString UK2Node_IfThenElse::GetKeywords() const
diff --git a/Engine/Source/Editor/BlueprintGraph/Private/K2Node_InputAction.cpp b/Engine/Source/Editor/BlueprintGraph/Private/K2Node_InputAction.cpp
--- a/Engine/Source/Editor/BlueprintGraph/Private/K2Node_InputAction.cpp
+++ b/Engine/Source/Editor/BlueprintGraph/Private/K2Node_InputAction.cpp
@@ -64,9 +64,9 @@ FText UK2Node_InputAction::GetNodeTitle(ENodeTitleType::Type TitleType) const
 	return CachedNodeTitle;
 }
 
-FString UK2Node_InputAction::GetTooltip() const
+FText UK2Node_InputAction::GetTooltipText() const
 {
-	return FString::Printf(*NSLOCTEXT("K2Node", "InputAction_Tooltip", "Event for when the keys bound to input action %s are pressed or released.").ToString(), *InputActionName.ToString());
+	return FText::Format(NSLOCTEXT("K2Node", "InputAction_Tooltip", "Event for when the keys bound to input action {0} are pressed or released."), FText::FromName(InputActionName));
 }
 
 bool UK2Node_InputAction::IsCompatibleWithGraph(UEdGraph const* Graph) const
diff --git a/Engine/Source/Editor/BlueprintGraph/Private/K2Node_InputAxisEvent.cpp b/Engine/Source/Editor/BlueprintGraph/Private/K2Node_InputAxisEvent.cpp
--- a/Engine/Source/Editor/BlueprintGraph/Private/K2Node_InputAxisEvent.cpp
+++ b/Engine/Source/Editor/BlueprintGraph/Private/K2Node_InputAxisEvent.cpp
@@ -56,9 +56,9 @@ FText UK2Node_InputAxisEvent::GetNodeTitle(ENodeTitleType::Type TitleType) const
 	return CachedNodeTitle;
 }
 
-FString UK2Node_InputAxisEvent::GetTooltip() const
+FText UK2Node_InputAxisEvent::GetTooltipText() const
 {
-	return FString::Printf(*NSLOCTEXT("K2Node", "InputAxis_Tooltip", "Event that provides the current value of the %s axis once per frame when input is enabled for the containing actor.").ToString(), *InputAxisName.ToString());
+	return FText::Format(NSLOCTEXT("K2Node", "InputAxis_Tooltip", "Event that provides the current value of the {0} axis once per frame when input is enabled for the containing actor."), FText::FromName(InputAxisName));
 }
 
 void UK2Node_InputAxisEvent::ValidateNodeDuringCompilation(class FCompilerResultsLog& MessageLog) const
diff --git a/Engine/Source/Editor/BlueprintGraph/Private/K2Node_InputAxisKeyEvent.cpp b/Engine/Source/Editor/BlueprintGraph/Private/K2Node_InputAxisKeyEvent.cpp
--- a/Engine/Source/Editor/BlueprintGraph/Private/K2Node_InputAxisKeyEvent.cpp
+++ b/Engine/Source/Editor/BlueprintGraph/Private/K2Node_InputAxisKeyEvent.cpp
@@ -34,9 +34,9 @@ FText UK2Node_InputAxisKeyEvent::GetNodeTitle(ENodeTitleType::Type TitleType) co
 	return AxisKey.GetDisplayName();
 }
 
-FString UK2Node_InputAxisKeyEvent::GetTooltip() const
+FText UK2Node_InputAxisKeyEvent::GetTooltipText() const
 {
-	return FString::Printf(*NSLOCTEXT("K2Node", "InputAxisKey_Tooltip", "Event that provides the current value of the %s axis once per frame when input is enabled for the containing actor.").ToString(), *AxisKey.GetDisplayName().ToString());
+	return FText::Format(NSLOCTEXT("K2Node", "InputAxisKey_Tooltip", "Event that provides the current value of the {0} axis once per frame when input is enabled for the containing actor."), AxisKey.GetDisplayName());
 }
 
 void UK2Node_InputAxisKeyEvent::ValidateNodeDuringCompilation(class FCompilerResultsLog& MessageLog) const
diff --git a/Engine/Source/Editor/BlueprintGraph/Private/K2Node_InputKey.cpp b/Engine/Source/Editor/BlueprintGraph/Private/K2Node_InputKey.cpp
--- a/Engine/Source/Editor/BlueprintGraph/Private/K2Node_InputKey.cpp
+++ b/Engine/Source/Editor/BlueprintGraph/Private/K2Node_InputKey.cpp
@@ -248,17 +248,17 @@ FText UK2Node_InputKey::GetNodeTitle(ENodeTitleType::Type TitleType) const
 	}
 }
 
-FString UK2Node_InputKey::GetTooltip() const
+FText UK2Node_InputKey::GetTooltipText() const
 {
 	FText ModifierText = GetModifierText();
 	FText KeyText = GetKeyText();
 
 	if ( !ModifierText.IsEmpty() )
 	{
-		return FText::Format(NSLOCTEXT("K2Node", "InputKey_Tooltip_Modifiers", "Events for when the {0} key is pressed or released while {1} is also held."), KeyText, ModifierText).ToString();
+		return FText::Format(NSLOCTEXT("K2Node", "InputKey_Tooltip_Modifiers", "Events for when the {0} key is pressed or released while {1} is also held."), KeyText, ModifierText);
 	}
 
-	return FText::Format(NSLOCTEXT("K2Node", "InputKey_Tooltip", "Events for when the {0} key is pressed or released."), KeyText).ToString();
+	return FText::Format(NSLOCTEXT("K2Node", "InputKey_Tooltip", "Events for when the {0} key is pressed or released."), KeyText);
 }
 
 FName UK2Node_InputKey::GetPaletteIcon(FLinearColor& OutColor) const
diff --git a/Engine/Source/Editor/BlueprintGraph/Private/K2Node_Knot.cpp b/Engine/Source/Editor/BlueprintGraph/Private/K2Node_Knot.cpp
--- a/Engine/Source/Editor/BlueprintGraph/Private/K2Node_Knot.cpp
+++ b/Engine/Source/Editor/BlueprintGraph/Private/K2Node_Knot.cpp
@@ -30,10 +30,10 @@ void UK2Node_Knot::AllocateDefaultPins()
 	UEdGraphPin* MyOutputPin = CreatePin(EGPD_Output, Schema->PC_Wildcard, FString(), nullptr, /*bIsArray=*/ false, /*bIsReference=*/ false, OutputPinName);
 }
 
-FString UK2Node_Knot::GetTooltip() const
+FText UK2Node_Knot::GetTooltipText() const
 {
 	//@TODO: Should pull the tooltip from the source pin
-	return LOCTEXT("KnotTooltip", "Reroute Node (reroutes wires)").ToString();
+	return LOCTEXT("KnotTooltip", "Reroute Node (reroutes wires)");
 }
 
 FText UK2Node_Knot::GetNodeTitle(ENodeTitleType::Type TitleType) const
@@ -60,7 +60,7 @@ void UK2Node_Knot::GetMenuEntries(struct FGraphContextMenuBuilder& ContextMenuBu
 
 		FString EmptyCategory;
 		FText MenuDesc = LOCTEXT("KnotMenuDescription", "Add Reroute Node...");
-		FString Tooltip = TemplateNode->GetTooltip();//@TODO: Make this work +LOCTEXT("KnotMenuExtraTooltip", "\nYou can also create one by Shift+Dragging off a pin").ToString();
+		FString Tooltip = TemplateNode->GetTooltipText().ToString();//@TODO: Make this work +LOCTEXT("KnotMenuExtraTooltip", "\nYou can also create one by Shift+Dragging off a pin").ToString();
 		FString Keywords = TemplateNode->GetKeywords();
 
 		TSharedPtr<FEdGraphSchemaAction_K2NewNode> NodeAction = FK2ActionMenuBuilder::AddNewNodeAction(ContextMenuBuilder, EmptyCategory, MenuDesc, Tooltip, 0, Keywords);
diff --git a/Engine/Source/Editor/BlueprintGraph/Private/K2Node_Literal.cpp b/Engine/Source/Editor/BlueprintGraph/Private/K2Node_Literal.cpp
--- a/Engine/Source/Editor/BlueprintGraph/Private/K2Node_Literal.cpp
+++ b/Engine/Source/Editor/BlueprintGraph/Private/K2Node_Literal.cpp
@@ -108,9 +108,9 @@ void UK2Node_Literal::PostReconstructNode()
 	}
 }
 
-FString UK2Node_Literal::GetTooltip() const
+FText UK2Node_Literal::GetTooltipText() const
 {
-	return NSLOCTEXT("K2Node", "Literal_Tooltip", "Stores a reference to an actor in the level").ToString();
+	return NSLOCTEXT("K2Node", "Literal_Tooltip", "Stores a reference to an actor in the level");
 }
 
 FText UK2Node_Literal::GetNodeTitle(ENodeTitleType::Type TitleType) const
diff --git a/Engine/Source/Editor/BlueprintGraph/Private/K2Node_LocalVariable.cpp b/Engine/Source/Editor/BlueprintGraph/Private/K2Node_LocalVariable.cpp
--- a/Engine/Source/Editor/BlueprintGraph/Private/K2Node_LocalVariable.cpp
+++ b/Engine/Source/Editor/BlueprintGraph/Private/K2Node_LocalVariable.cpp
@@ -15,14 +15,14 @@ UDEPRECATED_K2Node_LocalVariable::UDEPRECATED_K2Node_LocalVariable(const class F
 	CustomVariableName = TEXT("NewLocalVar");
 }
 
-FString UDEPRECATED_K2Node_LocalVariable::GetTooltip() const
+FText UDEPRECATED_K2Node_LocalVariable::GetTooltipText() const
 {
 	if(VariableTooltip.IsEmpty())
 	{
-		return Super::GetTooltip();
+		return Super::GetTooltipText();
 	}
 
-	return VariableTooltip.ToString();
+	return VariableTooltip;
 }
 
 FText UDEPRECATED_K2Node_LocalVariable::GetNodeTitle(ENodeTitleType::Type TitleType) const
diff --git a/Engine/Source/Editor/BlueprintGraph/Private/K2Node_MacroInstance.cpp b/Engine/Source/Editor/BlueprintGraph/Private/K2Node_MacroInstance.cpp
--- a/Engine/Source/Editor/BlueprintGraph/Private/K2Node_MacroInstance.cpp
+++ b/Engine/Source/Editor/BlueprintGraph/Private/K2Node_MacroInstance.cpp
@@ -101,24 +101,24 @@ void UK2Node_MacroInstance::AllocateDefaultPins()
 	}
 }
 
-FString UK2Node_MacroInstance::GetTooltip() const
+FText UK2Node_MacroInstance::GetTooltipText() const
 {
 	UEdGraph* MacroGraph = MacroGraphReference.GetGraph();
 	if (FKismetUserDeclaredFunctionMetadata* Metadata = GetAssociatedGraphMetadata(MacroGraph))
 	{
 		if (!Metadata->ToolTip.IsEmpty())
 		{
-			return Metadata->ToolTip;
+			return FText::FromString(Metadata->ToolTip);
 		}
 	}
 	
 	if (MacroGraph)
 	{
-		return MacroGraph->GetName() + NSLOCTEXT("K2Node", "Name_Instance", " instance").ToString();
+		return FText::Format(NSLOCTEXT("K2Node", "MacroGraphInstance_Tooltip", "{0} instance"), FText::FromName(MacroGraph->GetFName()));
 	}
 	else
 	{
-		return NSLOCTEXT("K2Node", "Macro_Tooltip", "Macro").ToString();
+		return NSLOCTEXT("K2Node", "Macro_Tooltip", "Macro");
 	}
 }
 
diff --git a/Engine/Source/Editor/BlueprintGraph/Private/K2Node_MakeArray.cpp b/Engine/Source/Editor/BlueprintGraph/Private/K2Node_MakeArray.cpp
--- a/Engine/Source/Editor/BlueprintGraph/Private/K2Node_MakeArray.cpp
+++ b/Engine/Source/Editor/BlueprintGraph/Private/K2Node_MakeArray.cpp
@@ -297,9 +297,9 @@ void UK2Node_MakeArray::PostReconstructNode()
 	}
 }
 
-FString UK2Node_MakeArray::GetTooltip() const
+FText UK2Node_MakeArray::GetTooltipText() const
 {
-	return LOCTEXT("MakeArrayTooltip", "Create an array from a series of items.").ToString();
+	return LOCTEXT("MakeArrayTooltip", "Create an array from a series of items.");
 }
 
 void UK2Node_MakeArray::AddInputPin()
diff --git a/Engine/Source/Editor/BlueprintGraph/Private/K2Node_MakeStruct.cpp b/Engine/Source/Editor/BlueprintGraph/Private/K2Node_MakeStruct.cpp
--- a/Engine/Source/Editor/BlueprintGraph/Private/K2Node_MakeStruct.cpp
+++ b/Engine/Source/Editor/BlueprintGraph/Private/K2Node_MakeStruct.cpp
@@ -162,12 +162,12 @@ FText UK2Node_MakeStruct::GetNodeTitle(ENodeTitleType::Type TitleType) const
 	return CachedNodeTitle;
 }
 
-FString UK2Node_MakeStruct::GetTooltip() const
+FText UK2Node_MakeStruct::GetTooltipText() const
 {
-	return FString::Printf(
-		*LOCTEXT("MakeStruct_Tooltip", "Adds a node that create a '%s' from its members").ToString(),
-		*(StructType ? StructType->GetName() : FString())
-		);
+	return FText::Format(
+		LOCTEXT("MakeStruct_Tooltip", "Adds a node that create a '{0}' from its members"),
+		StructType ? FText::FromName(StructType->GetFName()) : FText::GetEmpty()
+	);
 }
 
 FLinearColor UK2Node_MakeStruct::GetNodeTitleColor() const
diff --git a/Engine/Source/Editor/BlueprintGraph/Private/K2Node_MultiGate.cpp b/Engine/Source/Editor/BlueprintGraph/Private/K2Node_MultiGate.cpp
--- a/Engine/Source/Editor/BlueprintGraph/Private/K2Node_MultiGate.cpp
+++ b/Engine/Source/Editor/BlueprintGraph/Private/K2Node_MultiGate.cpp
@@ -462,9 +462,9 @@ UK2Node_MultiGate::UK2Node_MultiGate(const class FPostConstructInitializePropert
 {
 }
 
-FString UK2Node_MultiGate::GetTooltip() const
+FText UK2Node_MultiGate::GetTooltipText() const
 {
-	return NSLOCTEXT("K2Node", "MultiGate_Tooltip", "Executes a series of pins in order").ToString();
+	return NSLOCTEXT("K2Node", "MultiGate_Tooltip", "Executes a series of pins in order");
 }
 
 FLinearColor UK2Node_MultiGate::GetNodeTitleColor() const
diff --git a/Engine/Source/Editor/BlueprintGraph/Private/K2Node_Select.cpp b/Engine/Source/Editor/BlueprintGraph/Private/K2Node_Select.cpp
--- a/Engine/Source/Editor/BlueprintGraph/Private/K2Node_Select.cpp
+++ b/Engine/Source/Editor/BlueprintGraph/Private/K2Node_Select.cpp
@@ -256,9 +256,9 @@ void UK2Node_Select::AllocateDefaultPins()
 	Super::AllocateDefaultPins();
 }
 
-FString UK2Node_Select::GetTooltip() const
+FText UK2Node_Select::GetTooltipText() const
 {
-	return TEXT("Return the option at Index, (first option is indexed at 0)");
+	return LOCTEXT("SelectNodeTooltip", "Return the option at Index, (first option is indexed at 0)");
 }
 
 FText UK2Node_Select::GetNodeTitle(ENodeTitleType::Type TitleType) const
