Undo & Redo now properly affect Blackboard entries

Added RF_Transactional when editing blackboard data.

diff --git a/Engine/Source/Editor/BehaviorTreeEditor/Private/SBehaviorTreeBlackboardEditor.cpp b/Engine/Source/Editor/BehaviorTreeEditor/Private/SBehaviorTreeBlackboardEditor.cpp
--- a/Engine/Source/Editor/BehaviorTreeEditor/Private/SBehaviorTreeBlackboardEditor.cpp
+++ b/Engine/Source/Editor/BehaviorTreeEditor/Private/SBehaviorTreeBlackboardEditor.cpp
@@ -96,6 +96,7 @@ void SBehaviorTreeBlackboardEditor::HandleDeleteEntry()
 		if(BlackboardEntry != nullptr)
 		{
 			const FScopedTransaction Transaction(LOCTEXT("BlackboardEntryDeleteTransaction", "Delete Blackboard Entry"));
+			BlackboardData->SetFlags(RF_Transactional);
 			BlackboardData->Modify();
 		
 			for(int32 ItemIndex = 0; ItemIndex < BlackboardData->Keys.Num(); ItemIndex++)
@@ -169,6 +170,7 @@ void SBehaviorTreeBlackboardEditor::HandleKeyClassPicked(UClass* InClass)
 	check(InClass->IsChildOf(UBlackboardKeyType::StaticClass()));
 
 	const FScopedTransaction Transaction(LOCTEXT("BlackboardEntryAddTransaction", "Add Blackboard Entry"));
+	BlackboardData->SetFlags(RF_Transactional);
 	BlackboardData->Modify();
 
 	// create a name for this new key
diff --git a/Engine/Source/Editor/BehaviorTreeEditor/Private/SBehaviorTreeBlackboardView.cpp b/Engine/Source/Editor/BehaviorTreeEditor/Private/SBehaviorTreeBlackboardView.cpp
--- a/Engine/Source/Editor/BehaviorTreeEditor/Private/SBehaviorTreeBlackboardView.cpp
+++ b/Engine/Source/Editor/BehaviorTreeEditor/Private/SBehaviorTreeBlackboardView.cpp
@@ -184,6 +184,7 @@ private:
 		if(NewName != BlackboardEntryAction->Key.EntryName)
 		{
 			const FScopedTransaction Transaction(LOCTEXT("BlackboardEntryRenameTransaction", "Rename Blackboard Entry"));
+			BlackboardEntryAction->BlackboardData->SetFlags(RF_Transactional);
 			BlackboardEntryAction->BlackboardData->Modify();
 			BlackboardEntryAction->Key.EntryName = FName(*NewText.ToString());
 
