Currect 'current' Blackboard values are now displayed in Blackboard view

Previously we were only showing the latest 'debug' values, wihc only get propogated on state change. Now we display the 'live' Blackboard values when 'Current' is selected.

TTP# 345511 - BT: blackboard preview is not updating values until new BT task is executed

diff --git a/Engine/Source/Editor/BehaviorTreeEditor/Private/BehaviorTreeDebugger.cpp b/Engine/Source/Editor/BehaviorTreeEditor/Private/BehaviorTreeDebugger.cpp
--- a/Engine/Source/Editor/BehaviorTreeEditor/Private/BehaviorTreeDebugger.cpp
+++ b/Engine/Source/Editor/BehaviorTreeEditor/Private/BehaviorTreeDebugger.cpp
@@ -1044,19 +1044,19 @@ FText FBehaviorTreeDebugger::FindValueForKey(const FName& InKeyName, bool bUseCu
 	if (IsDebuggerRunning() &&
 		TreeInstance.IsValid())
 	{
-		FBehaviorTreeExecutionStep* ShowStep = nullptr;
-		if(bUseCurrentState && TreeInstance->DebuggerSteps.Num() > 0)
+		const TMap<FName, FString>* MapToQuery = nullptr;
+		if(bUseCurrentState)
 		{
-			ShowStep = &TreeInstance->DebuggerSteps[TreeInstance->DebuggerSteps.Num() - 1];
+			MapToQuery = &CurrentValues;
 		}
 		else if(TreeInstance->DebuggerSteps.IsValidIndex(ActiveStepIndex))
 		{
-			ShowStep = &TreeInstance->DebuggerSteps[ActiveStepIndex];
+			MapToQuery = &TreeInstance->DebuggerSteps[ActiveStepIndex].BlackboardValues;
 		}
 
-		if(ShowStep != nullptr)
+		if(MapToQuery != nullptr)
 		{
-			const FString* FindValue = ShowStep->BlackboardValues.Find(InKeyName);
+			const FString* FindValue = MapToQuery->Find(InKeyName);
 			if(FindValue != nullptr)
 			{
 				return FText::FromString(*FindValue);
