Reloading editor after undoing a delete & saving now correctly restores Behavior Tree nodes

RF_Transient is now cleared after undo

TTP# 346158 - AI: BEHAVIORTREE: Nodes do not get saved after undoing their delete

diff --git a/Engine/Source/Editor/BehaviorTreeEditor/Private/BehaviorTreeGraphNode.cpp b/Engine/Source/Editor/BehaviorTreeEditor/Private/BehaviorTreeGraphNode.cpp
--- a/Engine/Source/Editor/BehaviorTreeEditor/Private/BehaviorTreeGraphNode.cpp
+++ b/Engine/Source/Editor/BehaviorTreeEditor/Private/BehaviorTreeGraphNode.cpp
@@ -110,6 +110,7 @@ void UBehaviorTreeGraphNode::ResetNodeOwner()
 	{
 		UBehaviorTree* BT = Cast<UBehaviorTree>(GetBehaviorTreeGraph()->GetOuter());
 		NodeInstance->Rename(NULL, BT, REN_DontCreateRedirectors | REN_DoNotDirty);
+		NodeInstance->ClearFlags(RF_Transient);
 	}
 }
 
