Right-click menu operations on Behavior Tree sub-nodes now work

Added new GetNodeUnderMouse() call to SGraphNode. This gives each node type the option of returning child nodes, rather than always operating on the top-level node.

TTP# 346150 - AI: BEHAVIORTREE: Unable to delete decorator and service using right click menu

diff --git a/Engine/Source/Editor/BehaviorTreeEditor/Private/SGraphNode_BehaviorTree.cpp b/Engine/Source/Editor/BehaviorTreeEditor/Private/SGraphNode_BehaviorTree.cpp
--- a/Engine/Source/Editor/BehaviorTreeEditor/Private/SGraphNode_BehaviorTree.cpp
+++ b/Engine/Source/Editor/BehaviorTreeEditor/Private/SGraphNode_BehaviorTree.cpp
@@ -658,7 +658,7 @@ FReply SGraphNode_BehaviorTree::OnMouseUp(const FGeometry& SenderGeometry, const
 	return FReply::Unhandled();
 }
 
-TSharedPtr<SGraphNode> SGraphNode_BehaviorTree::GetSubNodeUnderCursor(const FGeometry& WidgetGeometry, const FPointerEvent& MouseEvent)
+TSharedPtr<SGraphNode> SGraphNode_BehaviorTree::GetSubNodeUnderCursor(const FGeometry& WidgetGeometry, const FPointerEvent& MouseEvent) const
 {
 	TSharedPtr<SGraphNode> ResultNode;
 
@@ -1206,4 +1206,10 @@ TArray<FOverlayWidgetInfo> SGraphNode_BehaviorTree::GetOverlayWidgets(bool bSele
 	return Widgets;
 }
 
+TSharedRef<SGraphNode> SGraphNode_BehaviorTree::GetNodeUnderMouse(const FGeometry& MyGeometry, const FPointerEvent& MouseEvent)
+{
+	TSharedPtr<SGraphNode> SubNode = GetSubNodeUnderCursor(MyGeometry, MouseEvent);
+	return SubNode.IsValid() ? SubNode.ToSharedRef() : StaticCastSharedRef<SGraphNode>(AsShared());
+}
+
 #undef LOCTEXT_NAMESPACE
\ No newline at end of file
diff --git a/Engine/Source/Editor/BehaviorTreeEditor/Private/SGraphNode_BehaviorTree.h b/Engine/Source/Editor/BehaviorTreeEditor/Private/SGraphNode_BehaviorTree.h
--- a/Engine/Source/Editor/BehaviorTreeEditor/Private/SGraphNode_BehaviorTree.h
+++ b/Engine/Source/Editor/BehaviorTreeEditor/Private/SGraphNode_BehaviorTree.h
@@ -21,6 +21,7 @@ public:
 	virtual void OnDragLeave( const FDragDropEvent& DragDropEvent ) override;
 	virtual void GetOverlayBrushes(bool bSelected, const FVector2D WidgetSize, TArray<FOverlayBrushInfo>& Brushes) const override;
 	virtual TArray<FOverlayWidgetInfo> GetOverlayWidgets(bool bSelected, const FVector2D& WidgetSize) const override;
+	virtual TSharedRef<SGraphNode> GetNodeUnderMouse(const FGeometry& MyGeometry, const FPointerEvent& MouseEvent) override;
 	// End of SGraphNode interface
 
 	/** handle double click */
@@ -45,7 +46,7 @@ public:
 	void Tick( const FGeometry& AllottedGeometry, const double InCurrentTime, const float InDeltaTime ) override;
 	
 	/** gets decorator or service node if one is found under mouse cursor */
-	TSharedPtr<SGraphNode> GetSubNodeUnderCursor(const FGeometry& WidgetGeometry, const FPointerEvent& MouseEvent);
+	TSharedPtr<SGraphNode> GetSubNodeUnderCursor(const FGeometry& WidgetGeometry, const FPointerEvent& MouseEvent) const;
 
 	/** adds decorator widget inside current node */
 	void AddDecorator(TSharedPtr<SGraphNode> DecoratorWidget);
diff --git a/Engine/Source/Editor/GraphEditor/Private/SGraphNode.cpp b/Engine/Source/Editor/GraphEditor/Private/SGraphNode.cpp
--- a/Engine/Source/Editor/GraphEditor/Private/SGraphNode.cpp
+++ b/Engine/Source/Editor/GraphEditor/Private/SGraphNode.cpp
@@ -523,6 +523,11 @@ UEdGraphNode* SGraphNode::GetNodeObj() const
 	return GraphNode;
 }
 
+TSharedRef<SGraphNode> SGraphNode::GetNodeUnderMouse(const FGeometry& MyGeometry, const FPointerEvent& MouseEvent)
+{
+	return StaticCastSharedRef<SGraphNode>(AsShared());
+}
+
 TSharedPtr<SGraphPanel> SGraphNode::GetOwnerPanel() const
 {
 	return OwnerGraphPanelPtr.Pin();
diff --git a/Engine/Source/Editor/GraphEditor/Private/SGraphPanel.cpp b/Engine/Source/Editor/GraphEditor/Private/SGraphPanel.cpp
--- a/Engine/Source/Editor/GraphEditor/Private/SGraphPanel.cpp
+++ b/Engine/Source/Editor/GraphEditor/Private/SGraphPanel.cpp
@@ -642,7 +642,9 @@ TSharedPtr<SWidget> SGraphPanel::OnSummonContextMenu(const FGeometry& MyGeometry
 			if (HoveredNodeIndex != INDEX_NONE)
 			{
 				const FArrangedWidget& HoveredNode = ArrangedNodes[HoveredNodeIndex];
-				const TSharedRef<SGraphNode>& GraphNode = StaticCastSharedRef<SGraphNode>(HoveredNode.Widget);
+				TSharedRef<SGraphNode> GraphNode = StaticCastSharedRef<SGraphNode>(HoveredNode.Widget);
+				TSharedPtr<SGraphNode> GraphSubNode = GraphNode->GetNodeUnderMouse(HoveredNode.Geometry, MouseEvent);
+				GraphNode = GraphSubNode.IsValid() ? GraphSubNode.ToSharedRef() : GraphNode;
 				NodeUnderCursor = GraphNode->GetNodeObj();
 
 				// Selection should switch to this code if it isn't already selected.
diff --git a/Engine/Source/Editor/GraphEditor/Public/SGraphNode.h b/Engine/Source/Editor/GraphEditor/Public/SGraphNode.h
--- a/Engine/Source/Editor/GraphEditor/Public/SGraphNode.h
+++ b/Engine/Source/Editor/GraphEditor/Public/SGraphNode.h
@@ -120,6 +120,9 @@ public:
 	/** @return the node being observed by this widget*/
 	UEdGraphNode* GetNodeObj() const;
 
+	/** @return the node under the mouse (either this node or one of its children) */
+	virtual TSharedRef<SGraphNode> GetNodeUnderMouse(const FGeometry& MyGeometry, const FPointerEvent& MouseEvent);
+
 	TSharedPtr<SGraphPanel> GetOwnerPanel() const;
 
 	/**
