Prevented decorators & services from becoming pink when deleted & undone

Made sure that the outer was correcly restored for decorators and services (as well as clearing RF_Transient).

TTP# 346154 - AI: BEHAVIORTREE: Node decorators turn pink after undoing delete

diff --git a/Engine/Source/Editor/BehaviorTreeEditor/Private/BehaviorTreeGraphNode.cpp b/Engine/Source/Editor/BehaviorTreeEditor/Private/BehaviorTreeGraphNode.cpp
--- a/Engine/Source/Editor/BehaviorTreeEditor/Private/BehaviorTreeGraphNode.cpp
+++ b/Engine/Source/Editor/BehaviorTreeEditor/Private/BehaviorTreeGraphNode.cpp
@@ -111,6 +111,25 @@ void UBehaviorTreeGraphNode::ResetNodeOwner()
 		UBehaviorTree* BT = Cast<UBehaviorTree>(GetBehaviorTreeGraph()->GetOuter());
 		NodeInstance->Rename(NULL, BT, REN_DontCreateRedirectors | REN_DoNotDirty);
 		NodeInstance->ClearFlags(RF_Transient);
+
+		// also reset decorators and services
+		for(auto& Decorator : Decorators)
+		{
+			if(Decorator->NodeInstance != nullptr)
+			{
+				Decorator->NodeInstance->Rename(NULL, BT, REN_DontCreateRedirectors | REN_DoNotDirty);
+				Decorator->NodeInstance->ClearFlags(RF_Transient);
+			}
+		}
+
+		for(auto& Service : Services)
+		{
+			if(Service->NodeInstance != nullptr)
+			{
+				Service->NodeInstance->Rename(NULL, BT, REN_DontCreateRedirectors | REN_DoNotDirty);
+				Service->NodeInstance->ClearFlags(RF_Transient);
+			}
+		}
 	}
 }
 
