Fix crash on shutdown on Mac - the UObject system is shutdown before modules are unloaded so be sure to check that your UObject is going to be valid before dereferencing it.

diff --git a/Engine/Source/Editor/BehaviorTreeEditor/Private/BehaviorTreeEditorTypes.cpp b/Engine/Source/Editor/BehaviorTreeEditor/Private/BehaviorTreeEditorTypes.cpp
--- a/Engine/Source/Editor/BehaviorTreeEditor/Private/BehaviorTreeEditorTypes.cpp
+++ b/Engine/Source/Editor/BehaviorTreeEditor/Private/BehaviorTreeEditorTypes.cpp
@@ -173,8 +173,12 @@ FClassBrowseHelper::~FClassBrowseHelper()
 		}
 	
 		// Unregister to have Populate called when a Blueprint is compiled.
-		GEditor->OnBlueprintCompiled().RemoveAll(this);
-		GEditor->OnClassPackageLoadedOrUnloaded().RemoveAll(this);
+		if ( UObjectInitialized() )
+		{
+			// GEditor can't have been destructed before we call this or we'll crash.
+			GEditor->OnBlueprintCompiled().RemoveAll(this);
+			GEditor->OnClassPackageLoadedOrUnloaded().RemoveAll(this);
+		}
 	}
 }
 
