Fixed Blackboard 'Add Key' menu size to prevent dynamic resizing

The issue is a generic problem with dynamic-sized content when placed in dynamic-sized windows, so to address this locally for this instance, we can embed the classviewer in a fixed sized box.

TTP# 346143 - AI: BLACKBOARD: New Key drop down menu is inconsistent

diff --git a/Engine/Source/Editor/BehaviorTreeEditor/Private/SBehaviorTreeBlackboardEditor.cpp b/Engine/Source/Editor/BehaviorTreeEditor/Private/SBehaviorTreeBlackboardEditor.cpp
--- a/Engine/Source/Editor/BehaviorTreeEditor/Private/SBehaviorTreeBlackboardEditor.cpp
+++ b/Engine/Source/Editor/BehaviorTreeEditor/Private/SBehaviorTreeBlackboardEditor.cpp
@@ -157,7 +157,13 @@ TSharedRef<SWidget> SBehaviorTreeBlackboardEditor::HandleCreateNewEntryMenu() co
 
 	FOnClassPicked OnPicked( FOnClassPicked::CreateRaw( this, &SBehaviorTreeBlackboardEditor::HandleKeyClassPicked ) );
 
-	return FModuleManager::LoadModuleChecked<FClassViewerModule>("ClassViewer").CreateClassViewer(Options, OnPicked);
+	return 
+		SNew(SBox)
+		.HeightOverride(240.0f)
+		.WidthOverride(200.0f)
+		[
+			FModuleManager::LoadModuleChecked<FClassViewerModule>("ClassViewer").CreateClassViewer(Options, OnPicked)
+		];
 }
 
 void SBehaviorTreeBlackboardEditor::HandleKeyClassPicked(UClass* InClass)
