Disable Delete and Rename entries in Blackboard editor for inherited keys

Check for inherted status instead of just whether items are selected

TTP# 344055 - AI: BehaviorTree: Inherited keys right click context menu should have the rename/delete options grayed out

diff --git a/Engine/Source/Editor/BehaviorTreeEditor/Private/SBehaviorTreeBlackboardEditor.cpp b/Engine/Source/Editor/BehaviorTreeEditor/Private/SBehaviorTreeBlackboardEditor.cpp
--- a/Engine/Source/Editor/BehaviorTreeEditor/Private/SBehaviorTreeBlackboardEditor.cpp
+++ b/Engine/Source/Editor/BehaviorTreeEditor/Private/SBehaviorTreeBlackboardEditor.cpp
@@ -29,13 +29,13 @@ void SBehaviorTreeBlackboardEditor::Construct(const FArguments& InArgs, TSharedR
 	CommandList->MapAction(
 		FBTBlackboardCommands::Get().DeleteEntry,
 		FExecuteAction::CreateSP(this, &SBehaviorTreeBlackboardEditor::HandleDeleteEntry),
-		FCanExecuteAction::CreateSP(this, &SBehaviorTreeBlackboardEditor::HasSelectedItems)
+		FCanExecuteAction::CreateSP(this, &SBehaviorTreeBlackboardEditor::CanDeleteEntry)
 		);
 
 	CommandList->MapAction(
 		FGenericCommands::Get().Rename,
 		FExecuteAction::CreateSP(this, &SBehaviorTreeBlackboardEditor::HandleRenameEntry),
-		FCanExecuteAction::CreateSP(this, &SBehaviorTreeBlackboardEditor::HasSelectedItems)
+		FCanExecuteAction::CreateSP(this, &SBehaviorTreeBlackboardEditor::CanRenameEntry)
 		);
 
 	InCommandList->Append(CommandList);
@@ -92,8 +92,9 @@ void SBehaviorTreeBlackboardEditor::HandleDeleteEntry()
 
 	if(!IsDebuggerActive())
 	{
-		FBlackboardEntry* BlackboardEntry = GetSelectedEntry();
-		if(BlackboardEntry != nullptr)
+		bool bIsInherited = false;
+		FBlackboardEntry* BlackboardEntry = GetSelectedEntry(bIsInherited);
+		if(BlackboardEntry != nullptr && !bIsInherited)
 		{
 			const FScopedTransaction Transaction(LOCTEXT("BlackboardEntryDeleteTransaction", "Delete Blackboard Entry"));
 			BlackboardData->SetFlags(RF_Transactional);
@@ -128,6 +129,36 @@ void SBehaviorTreeBlackboardEditor::HandleRenameEntry()
 	}
 }
 
+bool SBehaviorTreeBlackboardEditor::CanDeleteEntry() const
+{
+	if(!IsDebuggerActive())
+	{
+		bool bIsInherited = false;
+		FBlackboardEntry* BlackboardEntry = GetSelectedEntry(bIsInherited);
+		if(BlackboardEntry != nullptr)
+		{
+			return !bIsInherited;
+		}
+	}
+
+	return false;
+}
+
+bool SBehaviorTreeBlackboardEditor::CanRenameEntry() const
+{
+	if(!IsDebuggerActive())
+	{
+		bool bIsInherited = false;
+		FBlackboardEntry* BlackboardEntry = GetSelectedEntry(bIsInherited);
+		if(BlackboardEntry != nullptr)
+		{
+			return !bIsInherited;
+		}
+	}
+
+	return false;
+}
+
 class FBlackboardEntryClassFilter : public IClassViewerFilter
 {
 public:
diff --git a/Engine/Source/Editor/BehaviorTreeEditor/Private/SBehaviorTreeBlackboardEditor.h b/Engine/Source/Editor/BehaviorTreeEditor/Private/SBehaviorTreeBlackboardEditor.h
--- a/Engine/Source/Editor/BehaviorTreeEditor/Private/SBehaviorTreeBlackboardEditor.h
+++ b/Engine/Source/Editor/BehaviorTreeEditor/Private/SBehaviorTreeBlackboardEditor.h
@@ -48,4 +48,10 @@ private:
 
 	/** Delegate handler that disallows creating a new entry while the debugger is running */
 	bool CanCreateNewEntry() const;
+
+	/** Check whether the 'Delete' operation can be performed on the selected item  */
+	bool CanDeleteEntry() const;
+
+	/** Check whether the 'Rename' operation can be performed on the selected item  */
+	bool CanRenameEntry() const;
 };
\ No newline at end of file
diff --git a/Engine/Source/Editor/BehaviorTreeEditor/Private/SBehaviorTreeBlackboardView.cpp b/Engine/Source/Editor/BehaviorTreeEditor/Private/SBehaviorTreeBlackboardView.cpp
--- a/Engine/Source/Editor/BehaviorTreeEditor/Private/SBehaviorTreeBlackboardView.cpp
+++ b/Engine/Source/Editor/BehaviorTreeEditor/Private/SBehaviorTreeBlackboardView.cpp
@@ -465,11 +465,12 @@ int32 SBehaviorTreeBlackboardView::GetSelectedEntryIndex(bool& bOutIsInherited)
 	return INDEX_NONE;
 }
 
-FBlackboardEntry* SBehaviorTreeBlackboardView::GetSelectedEntry() const
+FBlackboardEntry* SBehaviorTreeBlackboardView::GetSelectedEntry(bool& bOutIsInherited) const
 {
 	TSharedPtr<FEdGraphSchemaAction_BlackboardEntry> Entry = GetSelectedEntryInternal();
 	if(Entry.IsValid())
 	{
+		bOutIsInherited = Entry->bIsInherited;
 		return &Entry->Key;
 	}
 
@@ -550,7 +551,8 @@ bool SBehaviorTreeBlackboardView::IsUsingSavedValues() const
 
 bool SBehaviorTreeBlackboardView::HasSelectedItems() const
 {
-	return GetSelectedEntry() != nullptr;
+	bool bIsInherited = false;
+	return GetSelectedEntry(bIsInherited) != nullptr;
 }
 
 bool SBehaviorTreeBlackboardView::IsDebuggerActive() const
diff --git a/Engine/Source/Editor/BehaviorTreeEditor/Private/SBehaviorTreeBlackboardView.h b/Engine/Source/Editor/BehaviorTreeEditor/Private/SBehaviorTreeBlackboardView.h
--- a/Engine/Source/Editor/BehaviorTreeEditor/Private/SBehaviorTreeBlackboardView.h
+++ b/Engine/Source/Editor/BehaviorTreeEditor/Private/SBehaviorTreeBlackboardView.h
@@ -57,10 +57,10 @@ public:
 
 	/**
 	 * Retrieves the blackboard item currently selected by the user.
-	 * 
+	 * @param	bOutIsInherited		Whether the item's index is in the key or parent key array
 	 * @return A pointer to the currently selected entry (nullptr if a category is selected, or nothing at all)
 	 */
-	FBlackboardEntry* GetSelectedEntry() const;
+	FBlackboardEntry* GetSelectedEntry(bool& bOutIsInherited) const;
 
 	/**
 	 * Retrieves the blackboard item index currently selected by the user.
