Do not crash on node double-click when a BT-related asset is not found

Adding some extra safety around the double-click code.

TTP# 347500 - [BTEditor] Issues with BP-implemented BT nodes

diff --git a/Engine/Source/Editor/BehaviorTreeEditor/Private/BehaviorTreeEditor.cpp b/Engine/Source/Editor/BehaviorTreeEditor/Private/BehaviorTreeEditor.cpp
--- a/Engine/Source/Editor/BehaviorTreeEditor/Private/BehaviorTreeEditor.cpp
+++ b/Engine/Source/Editor/BehaviorTreeEditor/Private/BehaviorTreeEditor.cpp
@@ -1526,7 +1526,7 @@ void FBehaviorTreeEditor::OnNodeDoubleClicked(class UEdGraphNode* Node)
 	if (MyNode && MyNode->bInjectedNode)
 	{
 		UBTTask_RunBehavior* SubtreeTask = MyNode->ParentNode ? Cast<UBTTask_RunBehavior>(MyNode->ParentNode->NodeInstance) : NULL;
-		if (SubtreeTask)
+		if (SubtreeTask && SubtreeTask->GetSubtreeAsset())
 		{
 			FAssetEditorManager::Get().OpenEditorForAsset(SubtreeTask->GetSubtreeAsset());
 
@@ -1600,8 +1600,10 @@ void FBehaviorTreeEditor::OnNodeDoubleClicked(class UEdGraphNode* Node)
 		UPackage* Pkg = NodeClass->GetOuterUPackage();
 		FString ClassName = NodeClass->GetName().LeftChop(2);
 		UBlueprint* BlueprintOb = FindObject<UBlueprint>(Pkg, *ClassName);
-
-		FAssetEditorManager::Get().OpenEditorForAsset(BlueprintOb);
+		if(BlueprintOb)
+		{
+			FAssetEditorManager::Get().OpenEditorForAsset(BlueprintOb);
+		}
 	}
 }
 
