Prevent Blackboard rename and delete hotkeys from operating on invisible widgets

Only allow rename and delete actions to take place when in the correct editor mode.

TTP# 347513 - AI: BehaviorTree: Can enable rename widget on blackboard keys when focus is in behavior tree mode

diff --git a/Engine/Source/Editor/BehaviorTreeEditor/Private/BehaviorTreeEditor.cpp b/Engine/Source/Editor/BehaviorTreeEditor/Private/BehaviorTreeEditor.cpp
--- a/Engine/Source/Editor/BehaviorTreeEditor/Private/BehaviorTreeEditor.cpp
+++ b/Engine/Source/Editor/BehaviorTreeEditor/Private/BehaviorTreeEditor.cpp
@@ -248,7 +248,8 @@ void FBehaviorTreeEditor::InitBehaviorTreeEditor( const EToolkitMode::Type Mode,
 				.OnIsDebuggerPaused(this, &FBehaviorTreeEditor::IsDebuggerPaused)
 				.OnGetDebugTimeStamp(this, &FBehaviorTreeEditor::HandleGetDebugTimeStamp)
 				.OnGetDisplayCurrentState(this, &FBehaviorTreeEditor::HandleGetDisplayCurrentState)
-				.OnBlackboardKeyChanged(this, &FBehaviorTreeEditor::HandleBlackboardKeyChanged);
+				.OnBlackboardKeyChanged(this, &FBehaviorTreeEditor::HandleBlackboardKeyChanged)
+				.OnIsBlackboardModeActive(this, &FBehaviorTreeEditor::HandleIsBlackboardModeActive);
 	}
 	else
 	{
@@ -397,6 +398,11 @@ void FBehaviorTreeEditor::HandleBlackboardKeyChanged(UBlackboardData* InBlackboa
 	}
 }
 
+bool FBehaviorTreeEditor::HandleIsBlackboardModeActive() const
+{
+	return GetCurrentMode() == BlackboardMode;
+}
+
 void FBehaviorTreeEditor::GetBlackboardSelectionInfo(int32& OutSelectionIndex, bool& bOutIsInherited) const
 {
 	OutSelectionIndex = CurrentBlackboardEntryIndex;
diff --git a/Engine/Source/Editor/BehaviorTreeEditor/Private/BehaviorTreeEditor.h b/Engine/Source/Editor/BehaviorTreeEditor/Private/BehaviorTreeEditor.h
--- a/Engine/Source/Editor/BehaviorTreeEditor/Private/BehaviorTreeEditor.h
+++ b/Engine/Source/Editor/BehaviorTreeEditor/Private/BehaviorTreeEditor.h
@@ -184,6 +184,9 @@ public:
 	/** Delegate handler for when a blackboard key changes (added, removed, renamed) */
 	void HandleBlackboardKeyChanged(UBlackboardData* InBlackboardData, FBlackboardEntry* const InKey);
 
+	/** Check whether blackboard mode is current */
+	bool HandleIsBlackboardModeActive() const;
+
 	/** Get the currently selected blackboard entry */
 	void GetBlackboardSelectionInfo(int32& OutSelectionIndex, bool& bOutIsInherited) const;
 
diff --git a/Engine/Source/Editor/BehaviorTreeEditor/Private/SBehaviorTreeBlackboardEditor.cpp b/Engine/Source/Editor/BehaviorTreeEditor/Private/SBehaviorTreeBlackboardEditor.cpp
--- a/Engine/Source/Editor/BehaviorTreeEditor/Private/SBehaviorTreeBlackboardEditor.cpp
+++ b/Engine/Source/Editor/BehaviorTreeEditor/Private/SBehaviorTreeBlackboardEditor.cpp
@@ -23,6 +23,7 @@ void SBehaviorTreeBlackboardEditor::Construct(const FArguments& InArgs, TSharedR
 	OnIsDebuggerPaused = InArgs._OnIsDebuggerPaused;
 	OnGetDebugTimeStamp = InArgs._OnGetDebugTimeStamp;
 	OnGetDisplayCurrentState = InArgs._OnGetDisplayCurrentState;
+	OnIsBlackboardModeActive = InArgs._OnIsBlackboardModeActive;
 
 	TSharedRef<FUICommandList> CommandList = MakeShareable(new FUICommandList);
 
@@ -131,7 +132,9 @@ void SBehaviorTreeBlackboardEditor::HandleRenameEntry()
 
 bool SBehaviorTreeBlackboardEditor::CanDeleteEntry() const
 {
-	if(!IsDebuggerActive())
+	const bool bModeActive = OnIsBlackboardModeActive.IsBound() && OnIsBlackboardModeActive.Execute();
+
+	if(!IsDebuggerActive() && bModeActive)
 	{
 		bool bIsInherited = false;
 		FBlackboardEntry* BlackboardEntry = GetSelectedEntry(bIsInherited);
@@ -146,7 +149,9 @@ bool SBehaviorTreeBlackboardEditor::CanDeleteEntry() const
 
 bool SBehaviorTreeBlackboardEditor::CanRenameEntry() const
 {
-	if(!IsDebuggerActive())
+	const bool bModeActive = OnIsBlackboardModeActive.IsBound() && OnIsBlackboardModeActive.Execute();
+
+	if(!IsDebuggerActive() && bModeActive)
 	{
 		bool bIsInherited = false;
 		FBlackboardEntry* BlackboardEntry = GetSelectedEntry(bIsInherited);
diff --git a/Engine/Source/Editor/BehaviorTreeEditor/Private/SBehaviorTreeBlackboardEditor.h b/Engine/Source/Editor/BehaviorTreeEditor/Private/SBehaviorTreeBlackboardEditor.h
--- a/Engine/Source/Editor/BehaviorTreeEditor/Private/SBehaviorTreeBlackboardEditor.h
+++ b/Engine/Source/Editor/BehaviorTreeEditor/Private/SBehaviorTreeBlackboardEditor.h
@@ -6,6 +6,9 @@
 #include "EdGraph/EdGraphSchema.h"
 #include "SBehaviorTreeBlackboardView.h"
 
+/** Delegate used to determine whether the Blackboard mode is active */
+DECLARE_DELEGATE_RetVal(bool, FOnIsBlackboardModeActive);
+
 /** Displays and edits blackboard entries */
 class SBehaviorTreeBlackboardEditor : public SBehaviorTreeBlackboardView
 {
@@ -19,6 +22,7 @@ public:
 		SLATE_EVENT(FOnIsDebuggerPaused, OnIsDebuggerPaused)
 		SLATE_EVENT(FOnGetDebugTimeStamp, OnGetDebugTimeStamp)
 		SLATE_EVENT(FOnBlackboardKeyChanged, OnBlackboardKeyChanged)
+		SLATE_EVENT(FOnIsBlackboardModeActive, OnIsBlackboardModeActive)
 
 	SLATE_END_ARGS()
 
@@ -54,4 +58,8 @@ private:
 
 	/** Check whether the 'Rename' operation can be performed on the selected item  */
 	bool CanRenameEntry() const;
+
+private:
+	/** Delegate used to determine whether the Blackboard mode is active */
+	FOnIsBlackboardModeActive OnIsBlackboardModeActive;
 };
\ No newline at end of file
