internal tasks & decorators of AITest project are now hidden in BT editor

diff --git a/Engine/Source/Developer/AITestSuite/Classes/BehaviorTree/TestBTDecorator_CantExecute.h b/Engine/Source/Developer/AITestSuite/Classes/BehaviorTree/TestBTDecorator_CantExecute.h
--- a/Engine/Source/Developer/AITestSuite/Classes/BehaviorTree/TestBTDecorator_CantExecute.h
+++ b/Engine/Source/Developer/AITestSuite/Classes/BehaviorTree/TestBTDecorator_CantExecute.h
@@ -4,7 +4,7 @@
 #include "BehaviorTree/BTDecorator.h"
 #include "TestBTDecorator_CantExecute.generated.h"
 
-UCLASS()
+UCLASS(meta=(HiddenNode))
 class UTestBTDecorator_CantExecute : public UBTDecorator
 {
 	GENERATED_UCLASS_BODY()
diff --git a/Engine/Source/Developer/AITestSuite/Classes/BehaviorTree/TestBTTask_Log.h b/Engine/Source/Developer/AITestSuite/Classes/BehaviorTree/TestBTTask_Log.h
--- a/Engine/Source/Developer/AITestSuite/Classes/BehaviorTree/TestBTTask_Log.h
+++ b/Engine/Source/Developer/AITestSuite/Classes/BehaviorTree/TestBTTask_Log.h
@@ -10,7 +10,7 @@ struct FBTLogTaskMemory
 	float RemainingWaitTime;
 };
 
-UCLASS()
+UCLASS(meta=(HiddenNode))
 class UTestBTTask_Log : public UBTTaskNode
 {
 	GENERATED_UCLASS_BODY()
diff --git a/Engine/Source/Developer/AITestSuite/Classes/BehaviorTree/TestBTTask_SetFlag.h b/Engine/Source/Developer/AITestSuite/Classes/BehaviorTree/TestBTTask_SetFlag.h
--- a/Engine/Source/Developer/AITestSuite/Classes/BehaviorTree/TestBTTask_SetFlag.h
+++ b/Engine/Source/Developer/AITestSuite/Classes/BehaviorTree/TestBTTask_SetFlag.h
@@ -4,7 +4,7 @@
 #include "BehaviorTree/BTTaskNode.h"
 #include "TestBTTask_SetFlag.generated.h"
 
-UCLASS()
+UCLASS(meta=(HiddenNode))
 class UTestBTTask_SetFlag : public UBTTaskNode
 {
 	GENERATED_UCLASS_BODY()
diff --git a/Engine/Source/Editor/BehaviorTreeEditor/Classes/BehaviorTreeEditorTypes.h b/Engine/Source/Editor/BehaviorTreeEditor/Classes/BehaviorTreeEditorTypes.h
--- a/Engine/Source/Editor/BehaviorTreeEditor/Classes/BehaviorTreeEditorTypes.h
+++ b/Engine/Source/Editor/BehaviorTreeEditor/Classes/BehaviorTreeEditorTypes.h
@@ -114,6 +114,7 @@ private:
 	void BuildClassGraph();
 	void AddClassGraphChildren(TSharedPtr<FClassDataNode> Node, TArray<TSharedPtr<FClassDataNode> >& NodeList);
 
+	bool IsHidingClass(UClass* Class);
 	bool IsHidingParentClass(UClass* Class);
 	bool IsPackageSaved(FName PackageName);
 
diff --git a/Engine/Source/Editor/BehaviorTreeEditor/Private/BehaviorTreeEditorTypes.cpp b/Engine/Source/Editor/BehaviorTreeEditor/Private/BehaviorTreeEditorTypes.cpp
--- a/Engine/Source/Editor/BehaviorTreeEditor/Private/BehaviorTreeEditorTypes.cpp
+++ b/Engine/Source/Editor/BehaviorTreeEditor/Private/BehaviorTreeEditorTypes.cpp
@@ -237,6 +237,12 @@ bool FClassBrowseHelper::IsHidingParentClass(UClass* Class)
 	return Class && Class->HasAnyClassFlags(CLASS_Native) && Class->HasMetaData(MetaHideParent);
 }
 
+bool FClassBrowseHelper::IsHidingClass(UClass* Class)
+{
+	static FName MetaHideInEditor = TEXT("HiddenNode");
+	return Class && Class->HasAnyClassFlags(CLASS_Native) && Class->HasMetaData(MetaHideInEditor);
+}
+
 bool FClassBrowseHelper::IsPackageSaved(FName PackageName)
 {
 	const bool bFound = FPackageName::SearchForPackageOnDisk(PackageName.ToString());
@@ -427,6 +433,8 @@ void FClassBrowseHelper::BuildClassGraph()
 				HideParentList.Add(TestClass->GetSuperClass());
 			}
 
+			NewData.bIsHidden = IsHidingClass(TestClass);
+
 			NewNode->Data = NewData;
 
 			if (TestClass == RootNodeClass)
