Subgraphs for composite decorators and run behavior nodes are now 'complex' preview tooltips

Subgraph/composites now display their graph in a 'complex' tooltip, further informing the user that they can be double-clicked and edited.

TTP# 346130  - AI: BEHAVIORTREE: RunBehavior node has inconsistent color

diff --git a/Engine/Source/Editor/BehaviorTreeEditor/Private/SGraphNode_BehaviorTree.cpp b/Engine/Source/Editor/BehaviorTreeEditor/Private/SGraphNode_BehaviorTree.cpp
--- a/Engine/Source/Editor/BehaviorTreeEditor/Private/SGraphNode_BehaviorTree.cpp
+++ b/Engine/Source/Editor/BehaviorTreeEditor/Private/SGraphNode_BehaviorTree.cpp
@@ -13,6 +13,7 @@
 #include "BehaviorTree/BTDecorator.h"
 #include "BehaviorTree/BehaviorTree.h"
 #include "BehaviorTree/Tasks/BTTask_RunBehavior.h"
+#include "IDocumentation.h"
 
 #define LOCTEXT_NAMESPACE "SGraphNode_BehaviorTree"
 
@@ -374,8 +375,7 @@ void SGraphNode_BehaviorTree::UpdateGraphNode()
 			.OnMouseButtonDown(this, &SGraphNode_BehaviorTree::OnMouseDown)
 			.OnMouseButtonUp(this, &SGraphNode_BehaviorTree::OnMouseUp)
 			[
-				SNew(SOverlay)
-				.ToolTipText( this, &SGraphNode::GetNodeTooltip )
+				SNew(SOverlay)	
 				// INPUT PIN AREA
 				+SOverlay::Slot()
 				.HAlign(HAlign_Fill)
@@ -953,19 +953,54 @@ TSharedPtr<SToolTip> SGraphNode_BehaviorTree::GetComplexTooltip()
 	{
 		return SNew(SToolTip)
 			[
-				// Create the tooltip graph preview, make sure to disable state overlays to
-				// prevent the PIE / read-only borders from obscuring the graph
-				SNew(SGraphPreviewer, DecoratorNode->GetBoundGraph())
-				.CornerOverlayText(this, &SGraphNode_BehaviorTree::GetPreviewCornerText)
-				.ShowGraphStateOverlay(false)
+				SNew(SOverlay)
+				+SOverlay::Slot()
+				[
+					// Create the tooltip graph preview, make sure to disable state overlays to
+					// prevent the PIE / read-only borders from obscuring the graph
+					SNew(SGraphPreviewer, DecoratorNode->GetBoundGraph())
+					.CornerOverlayText(LOCTEXT("CompositeDecoratorOverlayText", "Composite Decorator").ToString())
+					.ShowGraphStateOverlay(false)
+				]
+				+SOverlay::Slot()
+				.Padding(2.0f)
+				[
+					SNew(STextBlock)
+					.Text(LOCTEXT("CompositeDecoratorTooltip", "Double-click to Open"))
+					.ColorAndOpacity(FSlateColor::UseSubduedForeground())
+				]
 			];
 	}
-	return NULL;
-}
 
-FString SGraphNode_BehaviorTree::GetPreviewCornerText() const
-{
-	return FString(TEXT("Composite Decorator"));
+	UBehaviorTreeGraphNode_Task* TaskNode = Cast<UBehaviorTreeGraphNode_Task>(GraphNode);
+	if(TaskNode && TaskNode->NodeInstance)
+	{
+		UBTTask_RunBehavior* RunBehavior = Cast<UBTTask_RunBehavior>(TaskNode->NodeInstance);
+		if(RunBehavior && RunBehavior->GetSubtreeAsset() && RunBehavior->GetSubtreeAsset()->BTGraph)
+		{
+			return SNew(SToolTip)
+				[
+					SNew(SOverlay)
+					+SOverlay::Slot()
+					[
+						// Create the tooltip graph preview, make sure to disable state overlays to
+						// prevent the PIE / read-only borders from obscuring the graph
+						SNew(SGraphPreviewer, RunBehavior->GetSubtreeAsset()->BTGraph)
+						.CornerOverlayText(LOCTEXT("RunBehaviorOverlayText", "Run Behavior").ToString())
+						.ShowGraphStateOverlay(false)
+					]
+					+SOverlay::Slot()
+					.Padding(2.0f)
+					[
+						SNew(STextBlock)
+						.Text(LOCTEXT("RunBehaviorTooltip", "Double-click to Open"))
+						.ColorAndOpacity(FSlateColor::UseSubduedForeground())
+					]
+				];
+		}
+	}
+
+	return IDocumentation::Get()->CreateToolTip(TAttribute<FText>(this, &SGraphNode::GetNodeTooltip), NULL, GraphNode->GetDocumentationLink(), GraphNode->GetDocumentationExcerptName());
 }
 
 const FSlateBrush* SGraphNode_BehaviorTree::GetNameIcon() const
diff --git a/Engine/Source/Editor/BehaviorTreeEditor/Private/SGraphNode_BehaviorTree.h b/Engine/Source/Editor/BehaviorTreeEditor/Private/SGraphNode_BehaviorTree.h
--- a/Engine/Source/Editor/BehaviorTreeEditor/Private/SGraphNode_BehaviorTree.h
+++ b/Engine/Source/Editor/BehaviorTreeEditor/Private/SGraphNode_BehaviorTree.h
@@ -106,7 +106,6 @@ protected:
 	FSlateColor GetBackgroundColor() const;
 	FString	GetDescription() const;
 
-	virtual FString GetPreviewCornerText() const;
 	virtual const FSlateBrush* GetNameIcon() const;
 	virtual EVisibility GetBlueprintIconVisibility() const;
 
