Fixed blocked links in the Behavior Tree editor

Made more link conditions valid when editing behavior trees.

TTP# 345628 - EDITOR: BEHAVIOR TREES: Blocked from wiring an already wired node

diff --git a/Engine/Source/Editor/BehaviorTreeEditor/Private/EdGraphSchema_BehaviorTree.cpp b/Engine/Source/Editor/BehaviorTreeEditor/Private/EdGraphSchema_BehaviorTree.cpp
--- a/Engine/Source/Editor/BehaviorTreeEditor/Private/EdGraphSchema_BehaviorTree.cpp
+++ b/Engine/Source/Editor/BehaviorTreeEditor/Private/EdGraphSchema_BehaviorTree.cpp
@@ -593,13 +593,24 @@ const FPinConnectionResponse UEdGraphSchema_BehaviorTree::CanCreateConnection(co
 
 	const bool bPinASingleLink = bPinAIsSingleComposite || bPinAIsSingleTask || bPinAIsSingleNode;
 	const bool bPinBSingleLink = bPinBIsSingleComposite || bPinBIsSingleTask || bPinBIsSingleNode;
-	if ((bPinASingleLink && PinA->LinkedTo.Num() > 0) ||
-		(bPinBSingleLink && PinB->LinkedTo.Num() > 0))
+
+	if (PinB->LinkedTo.Num() > 0)
 	{
-		return FPinConnectionResponse(CONNECT_RESPONSE_DISALLOW, LOCTEXT("PinErrorSingleConnection", "Can't connect to multiple nodes"));
+		if(bPinASingleLink)
+		{
+			return FPinConnectionResponse(CONNECT_RESPONSE_BREAK_OTHERS_AB, LOCTEXT("PinConnectReplace", "Replace connection"));
+		}
+		else
+		{
+			return FPinConnectionResponse(CONNECT_RESPONSE_BREAK_OTHERS_B, LOCTEXT("PinConnectReplace", "Replace connection"));
+		}
 	}
 
-	if (PinB->LinkedTo.Num() > 0)
+	if (bPinASingleLink && PinA->LinkedTo.Num() > 0)
+	{
+		return FPinConnectionResponse(CONNECT_RESPONSE_BREAK_OTHERS_A, LOCTEXT("PinConnectReplace", "Replace connection"));
+	}
+	else if(bPinBSingleLink && PinB->LinkedTo.Num() > 0)
 	{
 		return FPinConnectionResponse(CONNECT_RESPONSE_BREAK_OTHERS_B, LOCTEXT("PinConnectReplace", "Replace connection"));
 	}
