Behavior Trees: Fixed crash dragging off an input node & creating a task

GetOutput() was returning NULL as task/leaf nodes have no outputs

TTP# 347894 - AI: BehaviorTree: CRASH: Dragging off of an input pin to create a "tasks" node will cause a crash

diff --git a/Engine/Source/Editor/BehaviorTreeEditor/Private/BehaviorTreeGraphNode.cpp b/Engine/Source/Editor/BehaviorTreeEditor/Private/BehaviorTreeGraphNode.cpp
--- a/Engine/Source/Editor/BehaviorTreeEditor/Private/BehaviorTreeGraphNode.cpp
+++ b/Engine/Source/Editor/BehaviorTreeEditor/Private/BehaviorTreeGraphNode.cpp
@@ -256,13 +256,15 @@ void UBehaviorTreeGraphNode::AutowireNewNode(UEdGraphPin* FromPin)
 {
 	Super::AutowireNewNode(FromPin);
 
-	if (FromPin != NULL)
+	if (FromPin != nullptr)
 	{
+		UEdGraphPin* OutputPin = GetOutputPin();
+
 		if (GetSchema()->TryCreateConnection(FromPin, GetInputPin()))
 		{
 			FromPin->GetOwningNode()->NodeConnectionListChanged();
 		}
-		else if(GetSchema()->TryCreateConnection(GetOutputPin(), FromPin))
+		else if(OutputPin != nullptr && GetSchema()->TryCreateConnection(OutputPin, FromPin))
 		{
 			NodeConnectionListChanged();
 		}
