Breaking behavior tree links is now undo-able

Added transactions for link breaking actions.
Removed call to unrelated sound schema delegate.

TTP# 347732 - EDITOR: BEHAVIOR TREES: Cannot undo Break Link action

diff --git a/Engine/Source/Editor/BehaviorTreeEditor/Classes/EdGraphSchema_BehaviorTree.h b/Engine/Source/Editor/BehaviorTreeEditor/Classes/EdGraphSchema_BehaviorTree.h
--- a/Engine/Source/Editor/BehaviorTreeEditor/Classes/EdGraphSchema_BehaviorTree.h
+++ b/Engine/Source/Editor/BehaviorTreeEditor/Classes/EdGraphSchema_BehaviorTree.h
@@ -106,6 +106,9 @@ class UEdGraphSchema_BehaviorTree : public UEdGraphSchema
 	virtual bool ShouldHidePinDefaultValue(UEdGraphPin* Pin) const override;
 	virtual class FConnectionDrawingPolicy* CreateConnectionDrawingPolicy(int32 InBackLayerID, int32 InFrontLayerID, float InZoomFactor, const FSlateRect& InClippingRect, class FSlateWindowElementList& InDrawElements, class UEdGraph* InGraphObj) const override;
 	virtual int32 GetNodeSelectionCount(const UEdGraph* Graph) const override;
+	virtual void BreakNodeLinks(UEdGraphNode& TargetNode) const override;
+	virtual void BreakPinLinks(UEdGraphPin& TargetPin, bool bSendsNodeNotification) const override;
+	virtual void BreakSinglePinLink(UEdGraphPin* SourcePin, UEdGraphPin* TargetPin) override;
 	// End EdGraphSchema interface
 
 	/* subnode creation context menu */
diff --git a/Engine/Source/Editor/BehaviorTreeEditor/Private/EdGraphSchema_BehaviorTree.cpp b/Engine/Source/Editor/BehaviorTreeEditor/Private/EdGraphSchema_BehaviorTree.cpp
--- a/Engine/Source/Editor/BehaviorTreeEditor/Private/EdGraphSchema_BehaviorTree.cpp
+++ b/Engine/Source/Editor/BehaviorTreeEditor/Private/EdGraphSchema_BehaviorTree.cpp
@@ -481,10 +481,30 @@ void UEdGraphSchema_BehaviorTree::GetBreakLinkToSubMenuActions( class FMenuBuild
 		++Count;
 
 		MenuBuilder.AddMenuEntry( Description, Description, FSlateIcon(), FUIAction(
-			FExecuteAction::CreateUObject((USoundClassGraphSchema*const)this, &USoundClassGraphSchema::BreakSinglePinLink, const_cast< UEdGraphPin* >(InGraphPin), *Links) ) );
+			FExecuteAction::CreateUObject(this, &UEdGraphSchema_BehaviorTree::BreakSinglePinLink, const_cast< UEdGraphPin* >(InGraphPin), *Links) ) );
 	}
 }
 
+void UEdGraphSchema_BehaviorTree::BreakNodeLinks(UEdGraphNode& TargetNode) const
+{
+	const FScopedTransaction Transaction( NSLOCTEXT("UnrealEd", "GraphEd_BreakNodeLinks", "Break Node Links") );
+
+	Super::BreakNodeLinks(TargetNode);
+}
+
+void UEdGraphSchema_BehaviorTree::BreakPinLinks(UEdGraphPin& TargetPin, bool bSendsNodeNotification) const
+{
+	const FScopedTransaction Transaction( NSLOCTEXT("UnrealEd", "GraphEd_BreakPinLinks", "Break Pin Links") );
+
+	Super::BreakPinLinks(TargetPin, bSendsNodeNotification);
+}
+
+void UEdGraphSchema_BehaviorTree::BreakSinglePinLink(UEdGraphPin* SourcePin, UEdGraphPin* TargetPin)
+{
+	const FScopedTransaction Transaction( NSLOCTEXT("UnrealEd", "GraphEd_BreakSinglePinLink", "Break Pin Link") );
+
+	Super::BreakSinglePinLink(SourcePin, TargetPin);
+}
 
 const FPinConnectionResponse UEdGraphSchema_BehaviorTree::CanCreateConnection(const UEdGraphPin* PinA, const UEdGraphPin* PinB) const
 {
