Fixed behavior tree pasting nodes offscreen

Average origin of pasted nodes was being calculated with the wrong divisor, so the position was miscalculated for nodes with any services or decorators.

TTP# 344625 - Editor: Copy/Paste of any node with a Decorator or Service in Behavior Trees pastes to the wrong location

diff --git a/Engine/Source/Editor/BehaviorTreeEditor/Private/BehaviorTreeEditor.cpp b/Engine/Source/Editor/BehaviorTreeEditor/Private/BehaviorTreeEditor.cpp
--- a/Engine/Source/Editor/BehaviorTreeEditor/Private/BehaviorTreeEditor.cpp
+++ b/Engine/Source/Editor/BehaviorTreeEditor/Private/BehaviorTreeEditor.cpp
@@ -1421,6 +1421,7 @@ void FBehaviorTreeEditor::PasteNodesHere(const FVector2D& Location)
 	//Average position of nodes so we can move them while still maintaining relative distances to each other
 	FVector2D AvgNodePosition(0.0f,0.0f);
 
+	float NumTopLevelNodes = 0.0f;
 	for (TSet<UEdGraphNode*>::TIterator It(PastedNodes); It; ++It)
 	{
 		UBehaviorTreeGraphNode* BTNode = Cast<UBehaviorTreeGraphNode>(*It);
@@ -1428,12 +1429,13 @@ void FBehaviorTreeEditor::PasteNodesHere(const FVector2D& Location)
 		{
 			AvgNodePosition.X += BTNode->NodePosX;
 			AvgNodePosition.Y += BTNode->NodePosY;
+			NumTopLevelNodes += 1.0f;
 		}
 	}
 
-	if (PastedNodes.Num() > 0)
+	if (NumTopLevelNodes > 0.0f)
 	{
-		float InvNumNodes = 1.0f/float(PastedNodes.Num());
+		float InvNumNodes = 1.0f/NumTopLevelNodes;
 		AvgNodePosition.X *= InvNumNodes;
 		AvgNodePosition.Y *= InvNumNodes;
 	}
