[NO NON-PII COMMIT MESSAGE]

diff --git a/Engine/Source/Editor/BehaviorTreeEditor/Private/SGraphNode_BehaviorTree.cpp b/Engine/Source/Editor/BehaviorTreeEditor/Private/SGraphNode_BehaviorTree.cpp
--- a/Engine/Source/Editor/BehaviorTreeEditor/Private/SGraphNode_BehaviorTree.cpp
+++ b/Engine/Source/Editor/BehaviorTreeEditor/Private/SGraphNode_BehaviorTree.cpp
@@ -206,7 +206,8 @@ FSlateColor SGraphNode_BehaviorTree::GetBorderBackgroundColor() const
 	UBehaviorTreeGraphNode* BTGraphNode = Cast<UBehaviorTreeGraphNode>(GraphNode);
 	const bool bIsInDebuggerActiveState = BTGraphNode && BTGraphNode->bDebuggerMarkCurrentlyActive;
 	const bool bIsInDebuggerPrevState = BTGraphNode && BTGraphNode->bDebuggerMarkPreviouslyActive;
-	const bool bSelectedSubNode = BTGraphNode && BTGraphNode->ParentNode && GetOwnerPanel()->SelectionManager.SelectedNodes.Contains(GraphNode);
+	const TSharedPtr<SGraphPanel> MyOwnerPanel = GetOwnerPanel();
+	const bool bSelectedSubNode = BTGraphNode && BTGraphNode->ParentNode && MyOwnerPanel.IsValid() && MyOwnerPanel->SelectionManager.SelectedNodes.Contains(GraphNode);
 	
 	UBTNode* NodeInstance = BTGraphNode ? Cast<UBTNode>(BTGraphNode->NodeInstance) : NULL;
 	const bool bIsDisconnected = NodeInstance && NodeInstance->GetExecutionIndex() == MAX_uint16;
