Fixed BTEditor crashing on selecting a service node #UE4

diff --git a/Engine/Source/Editor/BehaviorTreeEditor/Private/SGraphNode_BehaviorTree.cpp b/Engine/Source/Editor/BehaviorTreeEditor/Private/SGraphNode_BehaviorTree.cpp
--- a/Engine/Source/Editor/BehaviorTreeEditor/Private/SGraphNode_BehaviorTree.cpp
+++ b/Engine/Source/Editor/BehaviorTreeEditor/Private/SGraphNode_BehaviorTree.cpp
@@ -694,7 +694,7 @@ FReply SGraphNode_BehaviorTree::OnMouseDown(const FGeometry& SenderGeometry, con
 {
 	bIsMouseDown = true;
 
-	if (Cast<UBehaviorTreeGraphNode>(GraphNode)->ParentNode)
+	if (Cast<UBehaviorTreeGraphNode>(GraphNode)->ParentNode && GetOwnerPanel().IsValid())
 	{
 		GetOwnerPanel()->SelectionManager.ClickedOnNode(GraphNode,MouseEvent);
 		return FReply::Handled();
diff --git a/Engine/Source/Editor/GraphEditor/Private/SGraphNode.cpp b/Engine/Source/Editor/GraphEditor/Private/SGraphNode.cpp
--- a/Engine/Source/Editor/GraphEditor/Private/SGraphNode.cpp
+++ b/Engine/Source/Editor/GraphEditor/Private/SGraphNode.cpp
@@ -418,7 +418,7 @@ bool SGraphNode::IsSelectedExclusively() const
 {
 	TSharedPtr<SGraphPanel> OwnerPanel = OwnerGraphPanelPtr.Pin();
 
-	if(!OwnerPanel->HasKeyboardFocus() || OwnerPanel->SelectionManager.GetSelectedNodes().Num() > 1)
+	if (OwnerPanel.IsValid() == false || !OwnerPanel->HasKeyboardFocus() || OwnerPanel->SelectionManager.GetSelectedNodes().Num() > 1)
 	{
 		return false;
 	}
