Fixed crash selecting sub-nodes in Behavior Trees

Crash was caused by the removal of the call to UpdateGraphNode() in SGraphPanel::AddNode(). This meant that sub-nodes didnt get a valid parent panal pointer set.
Allowing nodes to override SetOwner() lets BT sub-nodes correctly recieve their parent.
Also removed previous checks put in place to try to prevent the crash.

diff --git a/Engine/Source/Editor/BehaviorTreeEditor/Private/SGraphNode_BehaviorTree.cpp b/Engine/Source/Editor/BehaviorTreeEditor/Private/SGraphNode_BehaviorTree.cpp
--- a/Engine/Source/Editor/BehaviorTreeEditor/Private/SGraphNode_BehaviorTree.cpp
+++ b/Engine/Source/Editor/BehaviorTreeEditor/Private/SGraphNode_BehaviorTree.cpp
@@ -207,8 +207,7 @@ FSlateColor SGraphNode_BehaviorTree::GetBorderBackgroundColor() const
 	UBehaviorTreeGraphNode* BTGraphNode = Cast<UBehaviorTreeGraphNode>(GraphNode);
 	const bool bIsInDebuggerActiveState = BTGraphNode && BTGraphNode->bDebuggerMarkCurrentlyActive;
 	const bool bIsInDebuggerPrevState = BTGraphNode && BTGraphNode->bDebuggerMarkPreviouslyActive;
-	const TSharedPtr<SGraphPanel> MyOwnerPanel = GetOwnerPanel();
-	const bool bSelectedSubNode = BTGraphNode && BTGraphNode->ParentNode && MyOwnerPanel.IsValid() && MyOwnerPanel->SelectionManager.SelectedNodes.Contains(GraphNode);
+	const bool bSelectedSubNode = BTGraphNode && BTGraphNode->ParentNode && GetOwnerPanel()->SelectionManager.SelectedNodes.Contains(GraphNode);
 	
 	UBTNode* NodeInstance = BTGraphNode ? Cast<UBTNode>(BTGraphNode->NodeInstance) : NULL;
 	const bool bIsDisconnected = NodeInstance && NodeInstance->GetExecutionIndex() == MAX_uint16;
@@ -694,7 +693,7 @@ FReply SGraphNode_BehaviorTree::OnMouseDown(const FGeometry& SenderGeometry, con
 {
 	bIsMouseDown = true;
 
-	if (Cast<UBehaviorTreeGraphNode>(GraphNode)->ParentNode && GetOwnerPanel().IsValid())
+	if (Cast<UBehaviorTreeGraphNode>(GraphNode)->ParentNode)
 	{
 		GetOwnerPanel()->SelectionManager.ClickedOnNode(GraphNode,MouseEvent);
 		return FReply::Handled();
@@ -1246,4 +1245,25 @@ TSharedRef<SGraphNode> SGraphNode_BehaviorTree::GetNodeUnderMouse(const FGeometr
 	return SubNode.IsValid() ? SubNode.ToSharedRef() : StaticCastSharedRef<SGraphNode>(AsShared());
 }
 
+void SGraphNode_BehaviorTree::SetOwner( const TSharedRef<SGraphPanel>& OwnerPanel )
+{
+	SGraphNode::SetOwner(OwnerPanel);
+
+	for(auto& DecoratorWidget : DecoratorWidgets)
+	{
+		if(DecoratorWidget.IsValid())
+		{
+			DecoratorWidget->SetOwner(OwnerPanel);
+		}
+	}
+
+	for(auto& ServicesWidget : ServicesWidgets)
+	{
+		if(ServicesWidget.IsValid())
+		{
+			ServicesWidget->SetOwner(OwnerPanel);
+		}
+	}
+}
+
 #undef LOCTEXT_NAMESPACE
\ No newline at end of file
diff --git a/Engine/Source/Editor/BehaviorTreeEditor/Private/SGraphNode_BehaviorTree.h b/Engine/Source/Editor/BehaviorTreeEditor/Private/SGraphNode_BehaviorTree.h
--- a/Engine/Source/Editor/BehaviorTreeEditor/Private/SGraphNode_BehaviorTree.h
+++ b/Engine/Source/Editor/BehaviorTreeEditor/Private/SGraphNode_BehaviorTree.h
@@ -22,6 +22,7 @@ public:
 	virtual void GetOverlayBrushes(bool bSelected, const FVector2D WidgetSize, TArray<FOverlayBrushInfo>& Brushes) const override;
 	virtual TArray<FOverlayWidgetInfo> GetOverlayWidgets(bool bSelected, const FVector2D& WidgetSize) const override;
 	virtual TSharedRef<SGraphNode> GetNodeUnderMouse(const FGeometry& MyGeometry, const FPointerEvent& MouseEvent) override;
+	virtual void SetOwner( const TSharedRef<SGraphPanel>& OwnerPanel ) override;
 	// End of SGraphNode interface
 
 	/** handle double click */
diff --git a/Engine/Source/Editor/GraphEditor/Private/SGraphNode.cpp b/Engine/Source/Editor/GraphEditor/Private/SGraphNode.cpp
--- a/Engine/Source/Editor/GraphEditor/Private/SGraphNode.cpp
+++ b/Engine/Source/Editor/GraphEditor/Private/SGraphNode.cpp
@@ -418,7 +418,7 @@ bool SGraphNode::IsSelectedExclusively() const
 {
 	TSharedPtr<SGraphPanel> OwnerPanel = OwnerGraphPanelPtr.Pin();
 
-	if (OwnerPanel.IsValid() == false || !OwnerPanel->HasKeyboardFocus() || OwnerPanel->SelectionManager.GetSelectedNodes().Num() > 1)
+	if (!OwnerPanel->HasKeyboardFocus() || OwnerPanel->SelectionManager.GetSelectedNodes().Num() > 1)
 	{
 		return false;
 	}
diff --git a/Engine/Source/Editor/GraphEditor/Public/SGraphNode.h b/Engine/Source/Editor/GraphEditor/Public/SGraphNode.h
--- a/Engine/Source/Editor/GraphEditor/Public/SGraphNode.h
+++ b/Engine/Source/Editor/GraphEditor/Public/SGraphNode.h
@@ -97,7 +97,7 @@ public:
 	void SetDoubleClickEvent(FSingleNodeEvent InDoubleClickEvent);
 
 	/** @param OwnerPanel  The GraphPanel that this node belongs to */
-	void SetOwner( const TSharedRef<SGraphPanel>& OwnerPanel );
+	virtual void SetOwner( const TSharedRef<SGraphPanel>& OwnerPanel );
 
 	/** @return the editable title for a node */
 	FString GetEditableNodeTitle() const;
