A fix for not being able to open composite BT decorators in BT editor merged over from dev branch #UE4

diff --git a/Engine/Source/Editor/BehaviorTreeEditor/Private/SGraphNode_BehaviorTree.cpp b/Engine/Source/Editor/BehaviorTreeEditor/Private/SGraphNode_BehaviorTree.cpp
--- a/Engine/Source/Editor/BehaviorTreeEditor/Private/SGraphNode_BehaviorTree.cpp
+++ b/Engine/Source/Editor/BehaviorTreeEditor/Private/SGraphNode_BehaviorTree.cpp
@@ -1280,6 +1280,7 @@ void SGraphNode_BehaviorTree::SetOwner( const TSharedRef<SGraphPanel>& OwnerPane
 		if(DecoratorWidget.IsValid())
 		{
 			DecoratorWidget->SetOwner(OwnerPanel);
+			OwnerPanel->AttachGraphEvents(DecoratorWidget);
 		}
 	}
 
@@ -1288,6 +1289,7 @@ void SGraphNode_BehaviorTree::SetOwner( const TSharedRef<SGraphPanel>& OwnerPane
 		if(ServicesWidget.IsValid())
 		{
 			ServicesWidget->SetOwner(OwnerPanel);
+			OwnerPanel->AttachGraphEvents(ServicesWidget);
 		}
 	}
 }
