Removed redundant condition from BTGraphHelpers::RebuildExecutionOrder #UE4

diff --git a/Engine/Source/Editor/BehaviorTreeEditor/Private/BehaviorTreeGraph.cpp b/Engine/Source/Editor/BehaviorTreeEditor/Private/BehaviorTreeGraph.cpp
--- a/Engine/Source/Editor/BehaviorTreeEditor/Private/BehaviorTreeGraph.cpp
+++ b/Engine/Source/Editor/BehaviorTreeEditor/Private/BehaviorTreeGraph.cpp
@@ -549,7 +549,7 @@ namespace BTGraphHelpers
 
 	void RebuildExecutionOrder(UBehaviorTreeGraphNode* RootEdNode, UBTCompositeNode* RootNode, uint16* ExecutionIndex, uint8 TreeDepth)
 	{
-		if (RootEdNode == NULL || RootEdNode == NULL)
+		if (RootEdNode == NULL)
 		{
 			return;
 		}
