Pressing enter now refreshes Behavior Tree search results

UE-3850 - AI: BehaviorTree: Pressing enter in search results does not refresh the list as seen in other parts of the editor

diff --git a/Engine/Source/Editor/BehaviorTreeEditor/Private/FindInBT.cpp b/Engine/Source/Editor/BehaviorTreeEditor/Private/FindInBT.cpp
--- a/Engine/Source/Editor/BehaviorTreeEditor/Private/FindInBT.cpp
+++ b/Engine/Source/Editor/BehaviorTreeEditor/Private/FindInBT.cpp
@@ -122,6 +122,7 @@ void SFindInBT::Construct( const FArguments& InArgs, TSharedPtr<FBehaviorTreeEdi
 					SAssignNew(SearchTextField, SSearchBox)
 					.HintText(LOCTEXT("BehaviorTreeSearchHint", "Enter text to find nodes..."))
 					.OnTextChanged(this, &SFindInBT::OnSearchTextChanged)
+					.OnTextCommitted(this, &SFindInBT::OnSearchTextCommitted)
 				]
 			]
 			+SVerticalBox::Slot()
@@ -160,6 +161,11 @@ void SFindInBT::OnSearchTextChanged(const FText& Text)
 	InitiateSearch();
 }
 
+void SFindInBT::OnSearchTextCommitted(const FText& Text, ETextCommit::Type CommitType)
+{
+	OnSearchTextChanged(Text);
+}
+
 void SFindInBT::InitiateSearch()
 {
 	TArray<FString> Tokens;
diff --git a/Engine/Source/Editor/BehaviorTreeEditor/Private/FindInBT.h b/Engine/Source/Editor/BehaviorTreeEditor/Private/FindInBT.h
--- a/Engine/Source/Editor/BehaviorTreeEditor/Private/FindInBT.h
+++ b/Engine/Source/Editor/BehaviorTreeEditor/Private/FindInBT.h
@@ -59,6 +59,9 @@ private:
 	/** Called when user changes the text they are searching for */
 	void OnSearchTextChanged(const FText& Text);
 
+	/** Called when user commits text */
+	void OnSearchTextCommitted(const FText& Text, ETextCommit::Type CommitType);
+
 	/** Get the children of a row */
 	void OnGetChildren(FSearchResult InItem, TArray<FSearchResult>& OutChildren);
 
