Fixed Blackboard icons being tinted incorrectly

Icons were being tinited twice - once with the icons tint and again with the same color via the ColorAndOpacity tint on their SImage.

UE-3232 - EDITOR: BLACKBOARD: Vector color not consistent

diff --git a/Engine/Source/Editor/BehaviorTreeEditor/Private/SBehaviorTreeBlackboardView.cpp b/Engine/Source/Editor/BehaviorTreeEditor/Private/SBehaviorTreeBlackboardView.cpp
--- a/Engine/Source/Editor/BehaviorTreeEditor/Private/SBehaviorTreeBlackboardView.cpp
+++ b/Engine/Source/Editor/BehaviorTreeEditor/Private/SBehaviorTreeBlackboardView.cpp
@@ -89,7 +89,7 @@ class SBehaviorTreeBlackboardItem : public SGraphPaletteItem
 		FSlateBrush const* IconBrush   = FEditorStyle::GetBrush(TEXT("NoBrush"));
 		GetPaletteItemIcon(GraphAction, IconBrush);
 
-		TSharedRef<SWidget> IconWidget = CreateIconWidget( GraphAction->TooltipDescription, IconBrush, IconBrush->TintColor );
+		TSharedRef<SWidget> IconWidget = CreateIconWidget( GraphAction->TooltipDescription, IconBrush, FLinearColor::White );
 		TSharedRef<SWidget> NameSlotWidget = CreateTextSlotWidget( NameFont, InCreateData, BlackboardEntryAction->bIsInherited );
 		TSharedRef<SWidget> DebugSlotWidget = CreateDebugSlotWidget( NameFont );
 
