Fixed code relying on SLATE_TEXT_ATTRIBUTE for STextBlock.

diff --git a/Engine/Source/Developer/ProjectLauncher/Private/Widgets/Launch/SProjectLauncherLaunchRoleEditor.h b/Engine/Source/Developer/ProjectLauncher/Private/Widgets/Launch/SProjectLauncherLaunchRoleEditor.h
--- a/Engine/Source/Developer/ProjectLauncher/Private/Widgets/Launch/SProjectLauncherLaunchRoleEditor.h
+++ b/Engine/Source/Developer/ProjectLauncher/Private/Widgets/Launch/SProjectLauncherLaunchRoleEditor.h
@@ -59,16 +59,16 @@ private:
 	}
 
 	// Callback for getting the name of the selected instance type.
-	FString HandleInstanceTypeComboButtonContentText( ) const
+	FText HandleInstanceTypeComboButtonContentText( ) const
 	{
 		ILauncherProfileLaunchRolePtr RolePtr = Role.Pin();
 
 		if (RolePtr.IsValid())
 		{
-			return ELauncherProfileRoleInstanceTypes::ToString(RolePtr->GetInstanceType());
+			return FText::FromString(ELauncherProfileRoleInstanceTypes::ToString(RolePtr->GetInstanceType()));
 		}
 
-		return FString();
+		return FText::GetEmpty();
 	}
 
 	// Callback for clicking an item in the 'Instance type' menu.
diff --git a/Engine/Source/Developer/ProjectLauncher/Private/Widgets/Package/SProjectLauncherPackagePage.cpp b/Engine/Source/Developer/ProjectLauncher/Private/Widgets/Package/SProjectLauncherPackagePage.cpp
--- a/Engine/Source/Developer/ProjectLauncher/Private/Widgets/Package/SProjectLauncherPackagePage.cpp
+++ b/Engine/Source/Developer/ProjectLauncher/Private/Widgets/Package/SProjectLauncherPackagePage.cpp
@@ -90,7 +90,7 @@ void SProjectLauncherPackagePage::Construct( const FArguments& InArgs, const FPr
 /* SProjectLauncherLaunchPage callbacks
  *****************************************************************************/
 
-FString SProjectLauncherPackagePage::HandlePackagingModeComboButtonContentText( ) const
+FText SProjectLauncherPackagePage::HandlePackagingModeComboButtonContentText( ) const
 {
 	ILauncherProfilePtr SelectedProfile = Model->GetSelectedProfile();
 
@@ -100,23 +100,23 @@ FString SProjectLauncherPackagePage::HandlePackagingModeComboButtonContentText(
 
 		if (PackagingMode == ELauncherProfilePackagingModes::DoNotPackage)
 		{
-			return TEXT("Do not package");
+			return LOCTEXT("DoNotPackageAction", "Do not package");
 		}
 
 		if (PackagingMode == ELauncherProfilePackagingModes::Locally)
 		{
-			return TEXT("Package & store locally");
+			return LOCTEXT("LocallyAction", "Package & store locally");
 		}
 
 		if (PackagingMode == ELauncherProfilePackagingModes::SharedRepository)
 		{
-			return TEXT("Package & store in repository");
+			return LOCTEXT("SharedRepositoryAction", "Package & store in repository");
 		}
 
-		return LOCTEXT("PackagingModeComboButtonDefaultText", "Select...").ToString();
+		return LOCTEXT("PackagingModeComboButtonDefaultText", "Select...");
 	}
 
-	return TEXT("");
+	return FText::GetEmpty();
 }
 
 
diff --git a/Engine/Source/Developer/ProjectLauncher/Private/Widgets/Package/SProjectLauncherPackagePage.h b/Engine/Source/Developer/ProjectLauncher/Private/Widgets/Package/SProjectLauncherPackagePage.h
--- a/Engine/Source/Developer/ProjectLauncher/Private/Widgets/Package/SProjectLauncherPackagePage.h
+++ b/Engine/Source/Developer/ProjectLauncher/Private/Widgets/Package/SProjectLauncherPackagePage.h
@@ -34,7 +34,7 @@ public:
 private:
 
 	// Callback for getting the content text of the 'Cook Mode' combo button.
-	FString HandlePackagingModeComboButtonContentText( ) const;
+	FText HandlePackagingModeComboButtonContentText( ) const;
 
 	// Callback for clicking an item in the 'Cook Mode' menu.
 	void HandlePackagingModeMenuEntryClicked( ELauncherProfilePackagingModes::Type PackagingMode );
diff --git a/Engine/Source/Developer/ProjectLauncher/Private/Widgets/Project/SProjectLauncherProjectPicker.cpp b/Engine/Source/Developer/ProjectLauncher/Private/Widgets/Project/SProjectLauncherProjectPicker.cpp
--- a/Engine/Source/Developer/ProjectLauncher/Private/Widgets/Project/SProjectLauncherProjectPicker.cpp
+++ b/Engine/Source/Developer/ProjectLauncher/Private/Widgets/Project/SProjectLauncherProjectPicker.cpp
@@ -108,31 +108,31 @@ TSharedRef<SWidget> SProjectLauncherProjectPicker::MakeProjectWidget()
 /* SProjectLauncherProjectPicker callbacks
  *****************************************************************************/
 
-FString SProjectLauncherProjectPicker::HandleProjectComboButtonText( ) const
+FText SProjectLauncherProjectPicker::HandleProjectComboButtonText( ) const
 {
 	if (LaunchProfileAttr.IsBound())
 	{
 		ILauncherProfilePtr Profile = LaunchProfileAttr.Get();
 		if (Profile.IsValid() && Profile->HasProjectSpecified())
 		{
-			return Profile->GetProjectName();
+			return FText::FromString(Profile->GetProjectName());
 		}
-		return TEXT("Any Project");
+		return LOCTEXT("AnyProjectAction", "Any Project");
 	}
 
 	FString Name = Model->GetProfileManager()->GetProjectName();
 	if (!Name.IsEmpty())
 	{
-		return Name;
+		return FText::FromString(Name);
 	}
 
-	return LOCTEXT("SelectProjectText", "Select...").ToString();
+	return LOCTEXT("SelectProjectText", "Select...");
 }
 
 
-FString SProjectLauncherProjectPicker::HandleProjectComboButtonToolTip( ) const
+FText SProjectLauncherProjectPicker::HandleProjectComboButtonToolTip( ) const
 {
-	return LOCTEXT("SelectProjectText", "Select or browse for a project").ToString();
+	return LOCTEXT("SelectProjectText", "Select or browse for a project");
 }
 
 
diff --git a/Engine/Source/Developer/ProjectLauncher/Private/Widgets/Project/SProjectLauncherProjectPicker.h b/Engine/Source/Developer/ProjectLauncher/Private/Widgets/Project/SProjectLauncherProjectPicker.h
--- a/Engine/Source/Developer/ProjectLauncher/Private/Widgets/Project/SProjectLauncherProjectPicker.h
+++ b/Engine/Source/Developer/ProjectLauncher/Private/Widgets/Project/SProjectLauncherProjectPicker.h
@@ -50,10 +50,10 @@ protected:
 private:
 
 	// Handles getting the text for the project combo button.
-	FString HandleProjectComboButtonText( ) const;
+	FText HandleProjectComboButtonText( ) const;
 
 	// Handles getting the tooltip for the project combo button.
-	FString HandleProjectComboButtonToolTip( ) const;
+	FText HandleProjectComboButtonToolTip( ) const;
 
 	// Handles clicking the "any project" option.
 	void HandleAnyProjectClicked(FString ProjectPath);
diff --git a/Engine/Source/Developer/ProjectLauncher/Private/Widgets/SProjectLauncherMessageListRow.h b/Engine/Source/Developer/ProjectLauncher/Private/Widgets/SProjectLauncherMessageListRow.h
--- a/Engine/Source/Developer/ProjectLauncher/Private/Widgets/SProjectLauncherMessageListRow.h
+++ b/Engine/Source/Developer/ProjectLauncher/Private/Widgets/SProjectLauncherMessageListRow.h
@@ -63,7 +63,7 @@ public:
 				[
 					SNew(STextBlock)
 						.ColorAndOpacity(HandleGetTextColor())
-						.Text(Message->Message)
+						.Text(FText::FromString(Message->Message))
 				];
 		}
 
diff --git a/Engine/Source/Developer/ProjectLauncher/Private/Widgets/SProjectLauncherPreviewPage.cpp b/Engine/Source/Developer/ProjectLauncher/Private/Widgets/SProjectLauncherPreviewPage.cpp
--- a/Engine/Source/Developer/ProjectLauncher/Private/Widgets/SProjectLauncherPreviewPage.cpp
+++ b/Engine/Source/Developer/ProjectLauncher/Private/Widgets/SProjectLauncherPreviewPage.cpp
@@ -1013,7 +1013,7 @@ TSharedRef<ITableRow> SProjectLauncherPreviewPage::HandleDeviceProxyListViewGene
 }
 
 
-FString SProjectLauncherPreviewPage::HandleInitialCultureTextBlockText( ) const
+FText SProjectLauncherPreviewPage::HandleInitialCultureTextBlockText( ) const
 {
 	ILauncherProfilePtr SelectedProfile = Model->GetSelectedProfile();
 
@@ -1023,17 +1023,17 @@ FString SProjectLauncherPreviewPage::HandleInitialCultureTextBlockText( ) const
 
 		if (InitialCulture.IsEmpty())
 		{
-			return LOCTEXT("DefaultText", "<default>").ToString();
+			return LOCTEXT("DefaultText", "<default>");
 		}
 
-		return InitialCulture;
+		return FText::FromString(InitialCulture);
 	}
 
-	return LOCTEXT("NotAvailableText", "n/a").ToString();
+	return LOCTEXT("NotAvailableText", "n/a");
 }
 
 
-FString SProjectLauncherPreviewPage::HandleInitialMapTextBlockText( ) const
+FText SProjectLauncherPreviewPage::HandleInitialMapTextBlockText( ) const
 {
 	ILauncherProfilePtr SelectedProfile = Model->GetSelectedProfile();
 
@@ -1043,13 +1043,13 @@ FString SProjectLauncherPreviewPage::HandleInitialMapTextBlockText( ) const
 
 		if (InitialMap.IsEmpty())
 		{
-			return LOCTEXT("DefaultText", "<default>").ToString();
+			return LOCTEXT("DefaultText", "<default>");
 		}
 
-		return InitialMap;
+		return FText::FromString(InitialMap);
 	}
 
-	return LOCTEXT("NotAvailableText", "n/a").ToString();
+	return LOCTEXT("NotAvailableText", "n/a");
 }
 
 
@@ -1069,7 +1069,7 @@ EVisibility SProjectLauncherPreviewPage::HandleLaunchSummaryBoxVisibility( ELaun
 }
 
 
-FString SProjectLauncherPreviewPage::HandleLaunchVsyncTextBlockText( ) const
+FText SProjectLauncherPreviewPage::HandleLaunchVsyncTextBlockText( ) const
 {
 	ILauncherProfilePtr SelectedProfile = Model->GetSelectedProfile();
 
@@ -1077,13 +1077,13 @@ FString SProjectLauncherPreviewPage::HandleLaunchVsyncTextBlockText( ) const
 	{
 		if (SelectedProfile->GetDefaultLaunchRole()->IsVsyncEnabled())
 		{
-			return LOCTEXT("YesText", "Yes").ToString();
+			return LOCTEXT("YesText", "Yes");
 		}
 
-		return LOCTEXT("NoText", "No").ToString();
+		return LOCTEXT("NoText", "No");
 	}
 
-	return LOCTEXT("NotAvailableText", "n/a").ToString();
+	return LOCTEXT("NotAvailableText", "n/a");
 }
 
 
@@ -1103,7 +1103,7 @@ EVisibility SProjectLauncherPreviewPage::HandlePackageSummaryBoxVisibility( ELau
 }
 
 
-FString SProjectLauncherPreviewPage::HandleProjectTextBlockText( ) const
+FText SProjectLauncherPreviewPage::HandleProjectTextBlockText( ) const
 {
 	ILauncherProfilePtr SelectedProfile = Model->GetSelectedProfile();
 
@@ -1113,17 +1113,17 @@ FString SProjectLauncherPreviewPage::HandleProjectTextBlockText( ) const
 
 		if (GameName.IsEmpty())
 		{
-			GameName = LOCTEXT("NotSetText", "<not set>").ToString();
+			return LOCTEXT("NotSetText", "<not set>");
 		}
 
-		return GameName;
+		return FText::FromString(GameName);
 	}
 
-	return LOCTEXT("NotAvailableText", "n/a").ToString();
+	return LOCTEXT("NotAvailableText", "n/a");
 }
 
 
-FString SProjectLauncherPreviewPage::HandleSelectedDeviceGroupTextBlockText( ) const
+FText SProjectLauncherPreviewPage::HandleSelectedDeviceGroupTextBlockText( ) const
 {
 	ILauncherProfilePtr SelectedProfile = Model->GetSelectedProfile();
 
@@ -1133,24 +1133,24 @@ FString SProjectLauncherPreviewPage::HandleSelectedDeviceGroupTextBlockText( ) c
 
 		if (SelectedGroup.IsValid())
 		{
-			return SelectedGroup->GetName();
+			return FText::FromString(SelectedGroup->GetName());
 		}
 	}
 		
-	return LOCTEXT("NoneText", "<none>").ToString();
+	return LOCTEXT("NoneText", "<none>");
 }
 
 
-FString SProjectLauncherPreviewPage::HandleSelectedProfileTextBlockText( ) const
+FText SProjectLauncherPreviewPage::HandleSelectedProfileTextBlockText( ) const
 {
 	ILauncherProfilePtr SelectedProfile = Model->GetSelectedProfile();
 
 	if (SelectedProfile.IsValid())
 	{
-		return SelectedProfile->GetName();
+		return FText::FromString(SelectedProfile->GetName());
 	}
 
-	return LOCTEXT("NoneText", "<none>").ToString();
+	return LOCTEXT("NoneText", "<none>");
 }
 
 
diff --git a/Engine/Source/Developer/ProjectLauncher/Private/Widgets/SProjectLauncherPreviewPage.h b/Engine/Source/Developer/ProjectLauncher/Private/Widgets/SProjectLauncherPreviewPage.h
--- a/Engine/Source/Developer/ProjectLauncher/Private/Widgets/SProjectLauncherPreviewPage.h
+++ b/Engine/Source/Developer/ProjectLauncher/Private/Widgets/SProjectLauncherPreviewPage.h
@@ -68,28 +68,28 @@ private:
 	TSharedRef<ITableRow> HandleDeviceProxyListViewGenerateRow( ITargetDeviceProxyPtr InItem, const TSharedRef<STableViewBase>& OwnerTable );
 
 	// Callback for getting the text in the 'Initial Culture' text block.
-	FString HandleInitialCultureTextBlockText( ) const;
+	FText HandleInitialCultureTextBlockText( ) const;
 
 	// Callback for getting the text in the 'Initial Map' text block.
-	FString HandleInitialMapTextBlockText( ) const;
+	FText HandleInitialMapTextBlockText( ) const;
 
 	// Callback for getting the visibility of the specified cook summary box.
 	EVisibility HandleLaunchSummaryBoxVisibility( ELauncherProfileLaunchModes::Type LaunchMode ) const;
 
 	// Callback for getting the text in the 'VSync' text block.
-	FString HandleLaunchVsyncTextBlockText( ) const;
+	FText HandleLaunchVsyncTextBlockText( ) const;
 
 	// Callback for getting the visibility of the specified cook summary box.
 	EVisibility HandlePackageSummaryBoxVisibility( ELauncherProfilePackagingModes::Type PackagingMode ) const;
 
 	// Callback for getting the text in the 'Game' text block.
-	FString HandleProjectTextBlockText( ) const;
+	FText HandleProjectTextBlockText( ) const;
 
 	// Callback for getting the name of the selected device group.
-	FString HandleSelectedDeviceGroupTextBlockText( ) const;
+	FText HandleSelectedDeviceGroupTextBlockText( ) const;
 
 	// Callback for getting the name of the selected profile.
-	FString HandleSelectedProfileTextBlockText( ) const;
+	FText HandleSelectedProfileTextBlockText( ) const;
 
 	// Callback for determining the visibility of a validation error icon.
 	EVisibility HandleValidationErrorIconVisibility( ELauncherProfileValidationErrors::Type Error ) const;
diff --git a/Engine/Source/Developer/ProjectLauncher/Private/Widgets/SProjectLauncherTaskListRow.cpp b/Engine/Source/Developer/ProjectLauncher/Private/Widgets/SProjectLauncherTaskListRow.cpp
--- a/Engine/Source/Developer/ProjectLauncher/Private/Widgets/SProjectLauncherTaskListRow.cpp
+++ b/Engine/Source/Developer/ProjectLauncher/Private/Widgets/SProjectLauncherTaskListRow.cpp
@@ -14,7 +14,7 @@ TSharedRef<SWidget> SProjectLauncherTaskListRow::GenerateWidgetForColumn(const F
 	{
 		return SNew(SBox)
 			.Padding(FMargin(4.0, 0.0))
-			.VAlign((VAlign_Center))
+			.VAlign(VAlign_Center)
 			[
 				SNew(STextBlock)
 				.Text(this, &SProjectLauncherTaskListRow::HandleDurationText)
@@ -26,7 +26,7 @@ TSharedRef<SWidget> SProjectLauncherTaskListRow::GenerateWidgetForColumn(const F
 
 			+ SOverlay::Slot()
 			.HAlign(HAlign_Center)
-			.VAlign((VAlign_Center))
+			.VAlign(VAlign_Center)
 			[
 				SNew(SThrobber)
 				.Animate(SThrobber::VerticalAndOpacity)
@@ -36,7 +36,7 @@ TSharedRef<SWidget> SProjectLauncherTaskListRow::GenerateWidgetForColumn(const F
 
 		+ SOverlay::Slot()
 			.HAlign(HAlign_Center)
-			.VAlign((VAlign_Center))
+			.VAlign(VAlign_Center)
 			[
 				SNew(SImage)
 				.ColorAndOpacity(this, &SProjectLauncherTaskListRow::HandleIconColorAndOpacity)
@@ -47,7 +47,7 @@ TSharedRef<SWidget> SProjectLauncherTaskListRow::GenerateWidgetForColumn(const F
 	{
 		return SNew(SBox)
 			.Padding(FMargin(4.0, 0.0))
-			.VAlign((VAlign_Center))
+			.VAlign(VAlign_Center)
 			[
 				SNew(STextBlock)
 				.Text(this, &SProjectLauncherTaskListRow::HandleStatusText)
@@ -61,10 +61,10 @@ TSharedRef<SWidget> SProjectLauncherTaskListRow::GenerateWidgetForColumn(const F
 		{
 			return SNew(SBox)
 				.Padding(FMargin(4.0, 0.0))
-				.VAlign((VAlign_Center))
+				.VAlign(VAlign_Center)
 				[
 					SNew(STextBlock)
-					.Text(TaskPtr->GetDesc())
+					.Text(FText::FromString(TaskPtr->GetDesc()))
 				];
 		}
 	}
diff --git a/Engine/Source/Developer/ProjectLauncher/Private/Widgets/SProjectLauncherTaskListRow.h b/Engine/Source/Developer/ProjectLauncher/Private/Widgets/SProjectLauncherTaskListRow.h
--- a/Engine/Source/Developer/ProjectLauncher/Private/Widgets/SProjectLauncherTaskListRow.h
+++ b/Engine/Source/Developer/ProjectLauncher/Private/Widgets/SProjectLauncherTaskListRow.h
@@ -47,7 +47,7 @@ public:
 private:
 
 	// Callback for getting the duration of the task.
-	FString HandleDurationText( ) const
+	FText HandleDurationText( ) const
 	{
 		ILauncherTaskPtr TaskPtr = Task.Pin();
 
@@ -56,11 +56,11 @@ private:
 			if ((TaskPtr->GetStatus() != ELauncherTaskStatus::Pending) &&
 				(TaskPtr->GetStatus() != ELauncherTaskStatus::Canceled))
 			{
-				return TaskPtr->GetDuration().ToString(TEXT("%h:%m:%s"));
+				return FText::AsTimespan(TaskPtr->GetDuration());
 			}
 		}
 
-		return FString();
+		return FText::GetEmpty();
 	}
 
 	// Callback for getting the color and opacity of the status icon.
@@ -116,7 +116,7 @@ private:
 	}
 
 	// Callback for getting the task's status text.
-	FString HandleStatusText( ) const
+	FText HandleStatusText( ) const
 	{
 		ILauncherTaskPtr TaskPtr = Task.Pin();
 
@@ -128,31 +128,31 @@ private:
 			{
 			case ELauncherTaskStatus::Busy:
 
-				return LOCTEXT("StatusInProgressText", "Busy").ToString();
+				return LOCTEXT("StatusInProgressText", "Busy");
 
 			case ELauncherTaskStatus::Canceled:
 
-				return LOCTEXT("StatusCanceledText", "Canceled").ToString();
+				return LOCTEXT("StatusCanceledText", "Canceled");
 
 			case ELauncherTaskStatus::Canceling:
 
-				return LOCTEXT("StatusCancelingText", "Canceling").ToString();
+				return LOCTEXT("StatusCancelingText", "Canceling");
 
 			case ELauncherTaskStatus::Completed:
 
-				return LOCTEXT("StatusCompletedText", "Completed").ToString();
+				return LOCTEXT("StatusCompletedText", "Completed");
 
 			case ELauncherTaskStatus::Failed:
 
-				return LOCTEXT("StatusFailedText", "Failed").ToString();
+				return LOCTEXT("StatusFailedText", "Failed");
 
 			case ELauncherTaskStatus::Pending:
 
-				return LOCTEXT("StatusPendingText", "Pending").ToString();
+				return LOCTEXT("StatusPendingText", "Pending");
 			}
 		}
 
-		return FString();
+		return FText::GetEmpty();
 	}
 
 	// Callback for determining the throbber's visibility.
diff --git a/Engine/Source/Developer/ProjectLauncher/Private/Widgets/SProjectLauncherValidation.h b/Engine/Source/Developer/ProjectLauncher/Private/Widgets/SProjectLauncherValidation.h
--- a/Engine/Source/Developer/ProjectLauncher/Private/Widgets/SProjectLauncherValidation.h
+++ b/Engine/Source/Developer/ProjectLauncher/Private/Widgets/SProjectLauncherValidation.h
@@ -136,7 +136,7 @@ protected:
 			.VAlign(VAlign_Center)
 			[
 				SNew(STextBlock)
-					.Text(MessageText)
+					.Text(FText::FromString(MessageText))
 			];
 	}
 
@@ -181,19 +181,19 @@ private:
 		return EVisibility::Collapsed;
 	}
 
-	FString	HandleValidationMessage( ELauncherProfileValidationErrors::Type Error ) const
+	FText HandleValidationMessage( ELauncherProfileValidationErrors::Type Error ) const
 	{
 		ILauncherProfilePtr LaunchProfile = LaunchProfileAttr.Get();
 		if (LaunchProfile.IsValid())
 		{
 			if (LaunchProfile->HasValidationError(Error))
 			{
-				return LOCTEXT("NoPlatformSDKInstalled", "A required platform SDK is mising: ").ToString() + LaunchProfile->GetInvalidPlatform();
+				return FText::Format(LOCTEXT("NoPlatformSDKInstalledFmt", "A required platform SDK is mising: {0}"), FText::FromString(LaunchProfile->GetInvalidPlatform()));
 			}
 
-			return TEXT("");
+			return FText::GetEmpty();
 		}
-		return LOCTEXT("InvalidLaunchProfile", "Invalid Launch Profile.").ToString();
+		return LOCTEXT("InvalidLaunchProfile", "Invalid Launch Profile.");
 	}
 
 private:
diff --git a/Engine/Source/Developer/ProjectLauncher/Private/Widgets/Shared/SProjectLauncherBuildConfigurationSelector.h b/Engine/Source/Developer/ProjectLauncher/Private/Widgets/Shared/SProjectLauncherBuildConfigurationSelector.h
--- a/Engine/Source/Developer/ProjectLauncher/Private/Widgets/Shared/SProjectLauncherBuildConfigurationSelector.h
+++ b/Engine/Source/Developer/ProjectLauncher/Private/Widgets/Shared/SProjectLauncherBuildConfigurationSelector.h
@@ -24,7 +24,7 @@ public:
 
 	SLATE_BEGIN_ARGS(SProjectLauncherBuildConfigurationSelector) { }
 		SLATE_EVENT(FOnSessionSProjectLauncherBuildConfigurationSelected, OnConfigurationSelected)
-		SLATE_ATTRIBUTE(FString, Text)
+		SLATE_ATTRIBUTE(FText, Text)
 		SLATE_ATTRIBUTE(FSlateFontInfo, Font)
 	SLATE_END_ARGS()
 
diff --git a/Engine/Source/Developer/ProjectLauncher/Private/Widgets/Shared/SProjectLauncherCookModeSelector.h b/Engine/Source/Developer/ProjectLauncher/Private/Widgets/Shared/SProjectLauncherCookModeSelector.h
--- a/Engine/Source/Developer/ProjectLauncher/Private/Widgets/Shared/SProjectLauncherCookModeSelector.h
+++ b/Engine/Source/Developer/ProjectLauncher/Private/Widgets/Shared/SProjectLauncherCookModeSelector.h
@@ -24,7 +24,7 @@ public:
 
 	SLATE_BEGIN_ARGS(SProjectLauncherCookModeSelector) { }
 		SLATE_EVENT(FOnSProjectLauncherCookModeSelected, OnCookModeSelected)
-		SLATE_ATTRIBUTE(FString, Text)
+		SLATE_ATTRIBUTE(FText, Text)
 	SLATE_END_ARGS()
 
 public:
diff --git a/Engine/Source/Developer/RealtimeProfiler/Private/SRealtimeProfilerTimeline.cpp b/Engine/Source/Developer/RealtimeProfiler/Private/SRealtimeProfilerTimeline.cpp
--- a/Engine/Source/Developer/RealtimeProfiler/Private/SRealtimeProfilerTimeline.cpp
+++ b/Engine/Source/Developer/RealtimeProfiler/Private/SRealtimeProfilerTimeline.cpp
@@ -53,7 +53,7 @@ void SRealtimeProfilerTimeline::Construct( const FArguments& InArgs )
 			.FillWidth( 1 )
 			.HAlign( HAlign_Fill )
 			[
-				SAssignNew( ZoomLabel, STextBlock ).Text( TAttribute<FString>( SharedThis( this ), &SRealtimeProfilerTimeline::GetZoomLabel ) )
+				SAssignNew( ZoomLabel, STextBlock ).Text( this, &SRealtimeProfilerTimeline::GetZoomLabel )
 			]
 			+SHorizontalBox::Slot()
 			.Padding( 2 )
@@ -101,9 +101,12 @@ void SRealtimeProfilerTimeline::ScrollBar_OnUserScrolled( float InScrollOffsetFr
 	}
 }
 
-FString SRealtimeProfilerTimeline::GetZoomLabel() const
+FText SRealtimeProfilerTimeline::GetZoomLabel() const
 {
-	return FString::Printf( TEXT( "%s %.2fx" ), *NSLOCTEXT("TaskGraph", "ZoomLabel", "Zoom").ToString(), GetZoom() );
+	static const FNumberFormattingOptions ZoomFormatOptions = FNumberFormattingOptions()
+		.SetMinimumFractionalDigits(2)
+		.SetMaximumFractionalDigits(2);
+	return FText::Format( NSLOCTEXT("TaskGraph", "ZoomLabelFmt", "Zoom: {0}x"), FText::AsNumber(GetZoom(), &ZoomFormatOptions) );
 }
 
 float SRealtimeProfilerTimeline::GetZoomValue() const
diff --git a/Engine/Source/Developer/RealtimeProfiler/Private/SRealtimeProfilerTimeline.h b/Engine/Source/Developer/RealtimeProfiler/Private/SRealtimeProfilerTimeline.h
--- a/Engine/Source/Developer/RealtimeProfiler/Private/SRealtimeProfilerTimeline.h
+++ b/Engine/Source/Developer/RealtimeProfiler/Private/SRealtimeProfilerTimeline.h
@@ -88,7 +88,7 @@ protected:
 	 *
 	 * @return Zoom label text
 	 */
-	FString GetZoomLabel() const;
+	FText GetZoomLabel() const;
 
 	/**
 	 * Callback used to get the current zoom slider value.
diff --git a/Engine/Source/Developer/ScreenShotComparison/Private/Widgets/SScreenPlatformRow.cpp b/Engine/Source/Developer/ScreenShotComparison/Private/Widgets/SScreenPlatformRow.cpp
--- a/Engine/Source/Developer/ScreenShotComparison/Private/Widgets/SScreenPlatformRow.cpp
+++ b/Engine/Source/Developer/ScreenShotComparison/Private/Widgets/SScreenPlatformRow.cpp
@@ -22,7 +22,7 @@ TSharedRef<SWidget> SScreenPlatformRow::GenerateWidgetForColumn( const FName& Co
 		.BorderImage(FEditorStyle::GetBrush("ToolPanel.GroupBorder"))
 		.Content()
 		[
-			SNew( STextBlock ) .Text( ScreenShotDataItem->GetName() )
+			SNew( STextBlock ) .Text( FText::FromString(ScreenShotDataItem->GetName()) )
 		];
 	}
 	else if ( ColumnName == TEXT( "Current View" ) )
diff --git a/Engine/Source/Developer/ScreenShotComparison/Private/Widgets/SScreenShotItem.cpp b/Engine/Source/Developer/ScreenShotComparison/Private/Widgets/SScreenShotItem.cpp
--- a/Engine/Source/Developer/ScreenShotComparison/Private/Widgets/SScreenShotItem.cpp
+++ b/Engine/Source/Developer/ScreenShotComparison/Private/Widgets/SScreenShotItem.cpp
@@ -31,7 +31,7 @@ void SScreenShotItem::Construct( const FArguments& InArgs )
 				.HAlign(HAlign_Center)
 				.Padding( 4.0f, 4.0f )
 				[
-					SNew( STextBlock ) .Text( ScreenShotData->GetName() )
+					SNew( STextBlock ) .Text( FText::FromString(ScreenShotData->GetName()) )
 				]
 				+SVerticalBox::Slot()
 				.AutoHeight()
diff --git a/Engine/Source/Developer/ScreenShotComparison/Private/Widgets/SScreenShotSearchBar.cpp b/Engine/Source/Developer/ScreenShotComparison/Private/Widgets/SScreenShotSearchBar.cpp
--- a/Engine/Source/Developer/ScreenShotComparison/Private/Widgets/SScreenShotSearchBar.cpp
+++ b/Engine/Source/Developer/ScreenShotComparison/Private/Widgets/SScreenShotSearchBar.cpp
@@ -129,7 +129,7 @@ void SScreenShotSearchBar::Construct( const FArguments& InArgs, IScreenShotManag
 				.OnClicked(this, &SScreenShotSearchBar::RefreshView)
 				[
 					SNew(SImage)
-					.ToolTipText( LOCTEXT( "Refresh Screen shots", "Refresh Screen Shots" ).ToString() )
+					.ToolTipText( LOCTEXT( "Refresh Screen shots", "Refresh Screen Shots" ) )
 					.Image(FEditorStyle::GetBrush("AutomationWindow.RefreshTests.Small"))
 				]
 			]
@@ -152,9 +152,9 @@ int32 SScreenShotSearchBar::GetDisplayEveryNthScreenshot() const
 	return DisplayEveryNthScreenshot;
 }
 
-FString SScreenShotSearchBar::GetPlatformString() const
+FText SScreenShotSearchBar::GetPlatformString() const
 {
-	return PlatformDisplayString;
+	return FText::FromString(PlatformDisplayString);
 }
 
 /** Filtering */
@@ -196,7 +196,7 @@ TSharedRef<ITableRow> SScreenShotSearchBar::HandlePlatformListViewGenerateRow( T
 	return SNew(STableRow<TSharedPtr<FString> >, OwnerTable)
 	.Content()
 	[
-		SNew( STextBlock ) .Text( *PlatformName )
+		SNew( STextBlock ) .Text( FText::FromString(*PlatformName) )
 	];
 }
 
diff --git a/Engine/Source/Developer/ScreenShotComparison/Private/Widgets/SScreenShotSearchBar.h b/Engine/Source/Developer/ScreenShotComparison/Private/Widgets/SScreenShotSearchBar.h
--- a/Engine/Source/Developer/ScreenShotComparison/Private/Widgets/SScreenShotSearchBar.h
+++ b/Engine/Source/Developer/ScreenShotComparison/Private/Widgets/SScreenShotSearchBar.h
@@ -30,7 +30,7 @@ protected:
 	 *
 	 * @return The platform string.
 	 */
-	FString GetPlatformString() const;
+	FText GetPlatformString() const;
 
 	/**
 	 * Set the view filter based on the input from the search box
diff --git a/Engine/Source/Developer/SessionFrontend/Private/Widgets/Browser/SSessionBrowser.cpp b/Engine/Source/Developer/SessionFrontend/Private/Widgets/Browser/SSessionBrowser.cpp
--- a/Engine/Source/Developer/SessionFrontend/Private/Widgets/Browser/SSessionBrowser.cpp
+++ b/Engine/Source/Developer/SessionFrontend/Private/Widgets/Browser/SSessionBrowser.cpp
@@ -42,7 +42,7 @@ void SSessionBrowser::Construct( const FArguments& InArgs, ISessionManagerRef In
 						SAssignNew(SessionComboBox, SComboBox<ISessionInfoPtr>)
 							.ContentPadding(FMargin(6.0f, 2.0f))
 							.OptionsSource(&SessionList)
-							.ToolTipText(LOCTEXT("SessionComboBoxToolTip", "Select the game session to interact with.").ToString())
+							.ToolTipText(LOCTEXT("SessionComboBoxToolTip", "Select the game session to interact with."))
 							.OnGenerateWidget(this, &SSessionBrowser::HandleSessionComboBoxGenerateWidget)
 							.OnSelectionChanged(this, &SSessionBrowser::HandleSessionComboBoxSelectionChanged)
 							[
@@ -78,7 +78,7 @@ void SSessionBrowser::Construct( const FArguments& InArgs, ISessionManagerRef In
 							.ContentPadding(FMargin(6.0f, 2.0f))
 							.IsEnabled(this, &SSessionBrowser::HandleTerminateSessionButtonIsEnabled)
 							.OnClicked(this, &SSessionBrowser::HandleTerminateSessionButtonClicked)
-							.ToolTipText(LOCTEXT("TerminateButtonTooltip", "Shuts down all game instances that are part of this session.").ToString())
+							.ToolTipText(LOCTEXT("TerminateButtonTooltip", "Shuts down all game instances that are part of this session."))
 							[
 								SNew(SHorizontalBox)
 
@@ -97,7 +97,7 @@ void SSessionBrowser::Construct( const FArguments& InArgs, ISessionManagerRef In
 									[
 										SNew(STextBlock)
 											.TextStyle(FEditorStyle::Get(), "SessionBrowser.Terminate.Font")
-											.Text(LOCTEXT("TerminateSessionButtonLabel", "Terminate Session").ToString())
+											.Text(LOCTEXT("TerminateSessionButtonLabel", "Terminate Session"))
 									]
 							]
 					]
@@ -226,7 +226,7 @@ TSharedRef<SWidget> SSessionBrowser::HandleSessionComboBoxGenerateWidget( ISessi
 				.HAlign(HAlign_Left)
 				[
 					SNew(STextBlock)
-						.Text(SessionInfo->GetSessionOwner())
+						.Text(FText::FromString(SessionInfo->GetSessionOwner()))
 				]
 
 			+ SHorizontalBox::Slot()
@@ -284,7 +284,7 @@ TSharedRef<ITableRow> SSessionBrowser::HandleSessionListViewGenerateRow( ISessio
 {
 	return SNew(SSessionBrowserSessionListRow, OwnerTable)
 		.SessionInfo(Item)
-		.ToolTipText(LOCTEXT("SessionToolTipSessionId", "Session ID: ").ToString() + Item->GetSessionId().ToString(EGuidFormats::DigitsWithHyphensInBraces));
+		.ToolTipText(FText::Format(LOCTEXT("SessionToolTipSessionIdFmt", "Session ID: {0}"), FText::FromString(Item->GetSessionId().ToString(EGuidFormats::DigitsWithHyphensInBraces))));
 }
 
 
diff --git a/Engine/Source/Developer/SessionFrontend/Private/Widgets/Browser/SSessionBrowserInstanceListRow.h b/Engine/Source/Developer/SessionFrontend/Private/Widgets/Browser/SSessionBrowserInstanceListRow.h
--- a/Engine/Source/Developer/SessionFrontend/Private/Widgets/Browser/SSessionBrowserInstanceListRow.h
+++ b/Engine/Source/Developer/SessionFrontend/Private/Widgets/Browser/SSessionBrowserInstanceListRow.h
@@ -88,7 +88,7 @@ public:
 						[
 							SNew(STextBlock)
 								.Font(FEditorStyle::GetFontStyle("BoldFont"))
-								.Text(InstanceInfo->GetInstanceName())							
+								.Text(FText::FromString(InstanceInfo->GetInstanceName()))
 						]
 				];
 		}
@@ -171,9 +171,9 @@ private:
 	}
 
 	/** Callback for getting the instance's current level. */
-	FString HandleLevelColumnText() const
+	FText HandleLevelColumnText() const
 	{
-		return InstanceInfo->GetCurrentLevel();
+		return FText::FromString(InstanceInfo->GetCurrentLevel());
 	}
 
 	/** Callback for getting the text in the 'Status' column. */
diff --git a/Engine/Source/Developer/SessionFrontend/Private/Widgets/Browser/SSessionBrowserSessionListRow.h b/Engine/Source/Developer/SessionFrontend/Private/Widgets/Browser/SSessionBrowserSessionListRow.h
--- a/Engine/Source/Developer/SessionFrontend/Private/Widgets/Browser/SSessionBrowserSessionListRow.h
+++ b/Engine/Source/Developer/SessionFrontend/Private/Widgets/Browser/SSessionBrowserSessionListRow.h
@@ -69,7 +69,7 @@ public:
 						.HAlign(HAlign_Left)
 						[
 							SNew(STextBlock)
-								.Text(SessionInfo->GetSessionOwner())
+								.Text(FText::FromString(SessionInfo->GetSessionOwner()))
 						]
 
 					+ SHorizontalBox::Slot()
@@ -92,14 +92,14 @@ public:
 protected:
 
 	/** Callback for getting the name of the session. */
-	FString HandleGetSessionName() const
+	FText HandleGetSessionName() const
 	{
 		const FString& SessionName = SessionInfo->GetSessionName();
 
 		// return name of a launched session
 		if (!SessionInfo->IsStandalone() || !SessionName.IsEmpty())
 		{
-			return SessionName;
+			return FText::FromString(SessionName);
 		}
 
 		// generate name for a standalone session
@@ -112,18 +112,18 @@ protected:
 
 			if ((Instances.Num() == 1) && (FirstInstance->GetInstanceId() == FApp::GetInstanceId()))
 			{
-				return LOCTEXT("ThisApplicationSessionText", "This Application").ToString();
+				return LOCTEXT("ThisApplicationSessionText", "This Application");
 			}
 
 			if (FirstInstance->GetDeviceName() == FPlatformProcess::ComputerName())
 			{
-				return LOCTEXT("UnnamedLocalSessionText", "Unnamed Session (Local)").ToString();
+				return LOCTEXT("UnnamedLocalSessionText", "Unnamed Session (Local)");
 			}
 
-			return LOCTEXT("UnnamedRemoteSessionText", "Unnamed Session (Remote)").ToString();
+			return LOCTEXT("UnnamedRemoteSessionText", "Unnamed Session (Remote)");
 		}
 
-		return LOCTEXT("UnnamedSessionText", "Unnamed Session").ToString();
+		return LOCTEXT("UnnamedSessionText", "Unnamed Session");
 	}
 
 private:
diff --git a/Engine/Source/Developer/SessionFrontend/Private/Widgets/Console/SSessionConsole.cpp b/Engine/Source/Developer/SessionFrontend/Private/Widgets/Console/SSessionConsole.cpp
--- a/Engine/Source/Developer/SessionFrontend/Private/Widgets/Console/SSessionConsole.cpp
+++ b/Engine/Source/Developer/SessionFrontend/Private/Widgets/Console/SSessionConsole.cpp
@@ -435,7 +435,7 @@ TSharedRef<ITableRow> SSessionConsole::HandleLogListGenerateRow( FSessionLogMess
 	return SNew(SSessionConsoleLogTableRow, OwnerTable)
 		.HighlightText(this, &SSessionConsole::HandleLogListGetHighlightText)
 		.LogMessage(Message)
-		.ToolTipText(Message->Text);
+		.ToolTipText(FText::FromString(Message->Text));
 }
 
 
diff --git a/Engine/Source/Developer/SessionFrontend/Private/Widgets/Console/SSessionConsoleFilterBar.cpp b/Engine/Source/Developer/SessionFrontend/Private/Widgets/Console/SSessionConsoleFilterBar.cpp
--- a/Engine/Source/Developer/SessionFrontend/Private/Widgets/Console/SSessionConsoleFilterBar.cpp
+++ b/Engine/Source/Developer/SessionFrontend/Private/Widgets/Console/SSessionConsoleFilterBar.cpp
@@ -186,9 +186,9 @@ TSharedRef<ITableRow> SSessionConsoleFilterBar::HandleCategoryFilterGenerateRow(
 }
 
 
-FString SSessionConsoleFilterBar::HandleCategoryFilterGetRowText( FSessionConsoleCategoryFilterPtr Filter ) const
+FText SSessionConsoleFilterBar::HandleCategoryFilterGetRowText( FSessionConsoleCategoryFilterPtr Filter ) const
 {
-	return FString::Printf(TEXT("%s (%i)"), *Filter->GetCategory().ToString(), CategoryCounters.FindRef(Filter->GetCategory()));
+	return FText::Format(LOCTEXT("CategoryFilterRowFmt", "{0} ({1})"), FText::FromName(Filter->GetCategory()), FText::AsNumber(CategoryCounters.FindRef(Filter->GetCategory())));
 }
 
 
@@ -219,9 +219,9 @@ void SSessionConsoleFilterBar::HandleHighlightOnlyCheckBoxCheckStateChanged( ECh
 }
 
 
-FString SSessionConsoleFilterBar::HandleVerbosityFilterGetRowText( FSessionConsoleVerbosityFilterPtr Filter ) const
+FText SSessionConsoleFilterBar::HandleVerbosityFilterGetRowText( FSessionConsoleVerbosityFilterPtr Filter ) const
 {
-	return FString::Printf(TEXT("%s (%i)"), *Filter->GetName(), VerbosityCounters.FindRef(Filter->GetVerbosity()));
+	return FText::Format(LOCTEXT("VerbosityFilterRowFmt", "{0} ({1})"), FText::FromString(Filter->GetName()), FText::AsNumber(VerbosityCounters.FindRef(Filter->GetVerbosity())));
 }
 
 BEGIN_SLATE_FUNCTION_BUILD_OPTIMIZATION
diff --git a/Engine/Source/Developer/SessionFrontend/Private/Widgets/Console/SSessionConsoleFilterBar.h b/Engine/Source/Developer/SessionFrontend/Private/Widgets/Console/SSessionConsoleFilterBar.h
--- a/Engine/Source/Developer/SessionFrontend/Private/Widgets/Console/SSessionConsoleFilterBar.h
+++ b/Engine/Source/Developer/SessionFrontend/Private/Widgets/Console/SSessionConsoleFilterBar.h
@@ -71,7 +71,7 @@ private:
 	TSharedRef<ITableRow> HandleCategoryFilterGenerateRow( FSessionConsoleCategoryFilterPtr Filter, const TSharedRef<STableViewBase>& OwnerTable );
 
 	/** Callback for getting the text for a row in the category filter drop-down. */
-	FString HandleCategoryFilterGetRowText( FSessionConsoleCategoryFilterPtr Filter ) const;
+	FText HandleCategoryFilterGetRowText( FSessionConsoleCategoryFilterPtr Filter ) const;
 
 	/** Callback for changing the enabled state of a category filter. */
 	void HandleCategoryFilterStateChanged( const FName& ChangedCategory, bool Enabled );
@@ -86,7 +86,7 @@ private:
 	TSharedRef<ITableRow> HandleVerbosityFilterGenerateRow( FSessionConsoleVerbosityFilterPtr Filter, const TSharedRef<STableViewBase>& OwnerTable );
 
 	/** Callback for getting the text for a row in the verbosity filter drop-down. */
-	FString HandleVerbosityFilterGetRowText( FSessionConsoleVerbosityFilterPtr Filter ) const;
+	FText HandleVerbosityFilterGetRowText( FSessionConsoleVerbosityFilterPtr Filter ) const;
 
 	/** Callback for changing the check state of a filter button. */
 	void HandleVerbosityFilterStateChanged( ELogVerbosity::Type Verbosity, bool Enabled );
diff --git a/Engine/Source/Developer/SessionFrontend/Private/Widgets/Console/SSessionConsoleLogTableRow.h b/Engine/Source/Developer/SessionFrontend/Private/Widgets/Console/SSessionConsoleLogTableRow.h
--- a/Engine/Source/Developer/SessionFrontend/Private/Widgets/Console/SSessionConsoleLogTableRow.h
+++ b/Engine/Source/Developer/SessionFrontend/Private/Widgets/Console/SSessionConsoleLogTableRow.h
@@ -52,7 +52,7 @@ public:
 						[
 							SNew(STextBlock)
 								.Font(FEditorStyle::GetFontStyle("BoldFont"))
-								.Text(LogMessage->InstanceName)
+								.Text(FText::FromString(LogMessage->InstanceName))
 						]
 				];
 		}
@@ -65,11 +65,15 @@ public:
 					SNew(STextBlock)
 						.ColorAndOpacity(HandleGetTextColor())
 						.HighlightText(HighlightText)
-						.Text(LogMessage->Text.Replace(TEXT("\n"), TEXT(" | ")).Replace(TEXT("\r"), TEXT("")))
+						.Text(FText::FromString(LogMessage->Text.Replace(TEXT("\n"), TEXT(" | ")).Replace(TEXT("\r"), TEXT(""))))
 				];
 		}
 		else if (ColumnName == "TimeSeconds")
 		{
+			static const FNumberFormattingOptions FormatOptions = FNumberFormattingOptions()
+				.SetMinimumFractionalDigits(3)
+				.SetMaximumFractionalDigits(3);
+
 			return SNew(SBox)
 				.HAlign(HAlign_Right)
 				.VAlign(VAlign_Center)
@@ -77,7 +81,7 @@ public:
 				[
 					SNew(STextBlock)
 						.ColorAndOpacity(HandleGetTextColor())
-						.Text(FString::Printf(TEXT("%.3f"), LogMessage->TimeSeconds))
+						.Text(FText::AsNumber(LogMessage->TimeSeconds, &FormatOptions))
 				];
 		}
 		else if (ColumnName == "Verbosity")
diff --git a/Engine/Source/Developer/SessionFrontend/Private/Widgets/Console/SSessionConsoleShortcutWindow.cpp b/Engine/Source/Developer/SessionFrontend/Private/Widgets/Console/SSessionConsoleShortcutWindow.cpp
--- a/Engine/Source/Developer/SessionFrontend/Private/Widgets/Console/SSessionConsoleShortcutWindow.cpp
+++ b/Engine/Source/Developer/SessionFrontend/Private/Widgets/Console/SSessionConsoleShortcutWindow.cpp
@@ -212,11 +212,11 @@ TSharedRef<ITableRow> SSessionConsoleShortcutWindow::HandleShortcutListViewGener
 					// execute button
 					SNew(SButton)
 						.VAlign(VAlign_Center)
-						.ToolTipText(InItem->Command)
+						.ToolTipText(FText::FromString(InItem->Command))
 						.OnClicked(FOnClicked::CreateSP(this, &SSessionConsoleShortcutWindow::HandleExecuteButtonClicked, InItem))
 						[
 							SNew(STextBlock)
-								.Text(InItem->Name)
+								.Text(FText::FromString(InItem->Name))
 						]
 				]
 
diff --git a/Engine/Source/Developer/SettingsEditor/Private/Widgets/SSettingsEditor.cpp b/Engine/Source/Developer/SettingsEditor/Private/Widgets/SSettingsEditor.cpp
--- a/Engine/Source/Developer/SettingsEditor/Private/Widgets/SSettingsEditor.cpp
+++ b/Engine/Source/Developer/SettingsEditor/Private/Widgets/SSettingsEditor.cpp
@@ -794,16 +794,17 @@ FText SSettingsEditor::HandleSettingsBoxDescriptionText() const
 }
 
 
-FString SSettingsEditor::HandleSettingsBoxTitleText() const
+FText SSettingsEditor::HandleSettingsBoxTitleText() const
 {
 	ISettingsSectionPtr SelectedSection = Model->GetSelectedSection();
 
 	if (SelectedSection.IsValid())
 	{
-		return SelectedSection->GetCategory().Pin()->GetDisplayName().ToString() + TEXT(" - ") + SelectedSection->GetDisplayName().ToString();
+		static const FText TitleFmt = FText::FromString(TEXT("{0} - {1}"));
+		return FText::Format(TitleFmt, SelectedSection->GetCategory().Pin()->GetDisplayName(), SelectedSection->GetDisplayName());
 	}
 
-	return FString();
+	return FText::GetEmpty();
 }
 
 
diff --git a/Engine/Source/Developer/SettingsEditor/Private/Widgets/SSettingsEditor.h b/Engine/Source/Developer/SettingsEditor/Private/Widgets/SSettingsEditor.h
--- a/Engine/Source/Developer/SettingsEditor/Private/Widgets/SSettingsEditor.h
+++ b/Engine/Source/Developer/SettingsEditor/Private/Widgets/SSettingsEditor.h
@@ -169,7 +169,7 @@ private:
 	FText HandleSettingsBoxDescriptionText() const;
 
 	/** Callback for getting the section text for the settings box. */
-	FString HandleSettingsBoxTitleText() const;
+	FText HandleSettingsBoxTitleText() const;
 
 	/** Callback for determining the visibility of the settings box. */
 	EVisibility HandleSettingsBoxVisibility() const;
diff --git a/Engine/Source/Developer/SuperSearch/Private/SSuperSearch.cpp b/Engine/Source/Developer/SuperSearch/Private/SSuperSearch.cpp
--- a/Engine/Source/Developer/SuperSearch/Private/SSuperSearch.cpp
+++ b/Engine/Source/Developer/SuperSearch/Private/SSuperSearch.cpp
@@ -228,7 +228,7 @@ TSharedRef<ITableRow> SSuperSearchBox::MakeSuggestionListItemWidget(TSharedPtr<F
 			.AutoWidth()
 			[
 				SNew(STextBlock)
-				.Text(Combined)
+				.Text(FText::FromString(Combined))
 #if WITH_EDITOR
 				.TextStyle(&FEditorStyle::Get().GetWidgetStyle<FTextBlockStyle>(bCategory ? "SuperSearchCategoryText" : "NormalText"))
 #endif
diff --git a/Engine/Source/Developer/TaskGraph/Private/SBarVisualizer.cpp b/Engine/Source/Developer/TaskGraph/Private/SBarVisualizer.cpp
--- a/Engine/Source/Developer/TaskGraph/Private/SBarVisualizer.cpp
+++ b/Engine/Source/Developer/TaskGraph/Private/SBarVisualizer.cpp
@@ -189,9 +189,12 @@ void SBarVisualizer::ScrollBar_OnUserScrolled( float InScrollOffsetFraction )
 	}
 }
 
-FString SBarVisualizer::GetZoomLabel() const
+FText SBarVisualizer::GetZoomLabel() const
 {
-	return FString::Printf( TEXT( "%s %.2fx" ), *NSLOCTEXT("TaskGraph", "ZoomLabel", "Zoom").ToString(), GetZoom() );
+	static const FNumberFormattingOptions ZoomFormatOptions = FNumberFormattingOptions()
+		.SetMinimumFractionalDigits(2)
+		.SetMaximumFractionalDigits(2);
+	return FText::Format( NSLOCTEXT("TaskGraph", "ZoomLabelFmt", "Zoom: {0}x"), FText::AsNumber(GetZoom(), &ZoomFormatOptions) );
 }
 
 float SBarVisualizer::GetZoomValue() const
@@ -394,12 +397,12 @@ TSharedRef<ITableRow> SBarVisualizer::OnGenerateWidgetForList( TSharedPtr<FVisua
 			+SHorizontalBox::Slot()
 			.AutoWidth()
 			[
-				SNew( STextBlock ).Text( InItem->EventName )
+				SNew( STextBlock ).Text( FText::FromString(InItem->EventName) )
 			];
 	}
 	else
 	{
-		BarTitle = SNew( STextBlock ).Text( InItem->EventName );
+		BarTitle = SNew( STextBlock ).Text( FText::FromString(InItem->EventName) );
 	}
 
 	return SNew( STableRow< TSharedPtr< FVisualizerEvent > >, OwnerTable )
@@ -604,7 +607,7 @@ void SBarVisualizer::OnBarGeometryChanged( FGeometry Geometry )
 	Timeline->SetDrawingGeometry( Geometry );
 }
 
-FString SBarVisualizer::GetSelectedCategoryName() const
+FText SBarVisualizer::GetSelectedCategoryName() const
 {
 	if( SelectedBarGraph.IsValid() )
 	{
@@ -614,11 +617,11 @@ FString SBarVisualizer::GetSelectedCategoryName() const
 		{
 			EventName = BarGraph->EventName + TEXT("\\") + EventName;
 		}
-		return EventName;
+		return FText::FromString(EventName);
 	}
 	else
 	{
-		return TEXT("Frame");
+		return NSLOCTEXT("SBarVisualizer", "Frame", "Frame");
 	}
 }
 
diff --git a/Engine/Source/Developer/TaskGraph/Private/SEventsTree.cpp b/Engine/Source/Developer/TaskGraph/Private/SEventsTree.cpp
--- a/Engine/Source/Developer/TaskGraph/Private/SEventsTree.cpp
+++ b/Engine/Source/Developer/TaskGraph/Private/SEventsTree.cpp
@@ -35,7 +35,7 @@ TSharedRef< SWidget > SEventItem::GenerateWidgetForColumn( const FName& ColumnNa
 		.AutoWidth()
 		[
 			SNew( STextBlock )
-			.Text( EventName )
+			.Text( FText::FromString(EventName) )
 		];
 	}
 	else if( ColumnName == SEventsTree::NAME_DurationColumn )
@@ -47,7 +47,7 @@ TSharedRef< SWidget > SEventItem::GenerateWidgetForColumn( const FName& ColumnNa
 	{
 		return
 		SNew( STextBlock )
-		. Text(FString::Printf( *LOCTEXT("UnsupportedColumn", "Unsupported Column: %s").ToString(), *(ColumnName.ToString()) ) );
+		. Text(FText::Format( LOCTEXT("UnsupportedColumnFmt", "Unsupported Column: {0}"), FText::FromName(ColumnName) ) );
 	}
 }
 END_SLATE_FUNCTION_BUILD_OPTIMIZATION
diff --git a/Engine/Source/Developer/TaskGraph/Public/SBarVisualizer.h b/Engine/Source/Developer/TaskGraph/Public/SBarVisualizer.h
--- a/Engine/Source/Developer/TaskGraph/Public/SBarVisualizer.h
+++ b/Engine/Source/Developer/TaskGraph/Public/SBarVisualizer.h
@@ -110,7 +110,7 @@ protected:
 	 *
 	 * @return Zoom label text
 	 */
-	FString GetZoomLabel() const;
+	FText GetZoomLabel() const;
 
 	/**
 	 * Callback used to get the current zoom slider value.
@@ -208,7 +208,7 @@ protected:
 	void OnBarGeometryChanged( FGeometry Geometry );
 
 	/** Gets the currently displayed hierarchy node name */
-	FString GetSelectedCategoryName() const;
+	FText GetSelectedCategoryName() const;
 
 	/** Checks if home button should be visible */
 	EVisibility GetHomeButtonVisibility() const;	
diff --git a/Engine/Source/Developer/TaskGraph/Public/SEventsTree.h b/Engine/Source/Developer/TaskGraph/Public/SEventsTree.h
--- a/Engine/Source/Developer/TaskGraph/Public/SEventsTree.h
+++ b/Engine/Source/Developer/TaskGraph/Public/SEventsTree.h
@@ -38,9 +38,12 @@ public:
 
 private:
 
-	FString	GetDurationText() const
+	FText GetDurationText() const
 	{
-		return FString::Printf( TEXT("%.2f"), EventDuration.Get() );
+		static const FNumberFormattingOptions DurationFormatOptions = FNumberFormattingOptions()
+			.SetMinimumFractionalDigits(2)
+			.SetMaximumFractionalDigits(2);
+		return FText::AsNumber(EventDuration.Get(), &DurationFormatOptions);
 	}
 
 	FString EventName;
diff --git a/Engine/Source/Editor/AddContentDialog/Private/SAddContentWidget.cpp b/Engine/Source/Editor/AddContentDialog/Private/SAddContentWidget.cpp
--- a/Engine/Source/Editor/AddContentDialog/Private/SAddContentWidget.cpp
+++ b/Engine/Source/Editor/AddContentDialog/Private/SAddContentWidget.cpp
@@ -236,7 +236,7 @@ TSharedRef<SWidget> SAddContentWidget::CreateContentSourceDetail(TSharedPtr<FCon
 			.Padding(FMargin(0, 0, 0, 5))
 			[
 				SNew(STextBlock)
-				.Text(ContentSource->GetClassTypes())
+				.Text(FText::FromString(ContentSource->GetClassTypes()))
 				.Visibility(ContentSource->GetClassTypes().IsEmpty() == false ? EVisibility::Visible : EVisibility::Collapsed)
 				.AutoWrapText(true)
 			]
diff --git a/Engine/Source/Editor/AnimGraph/Classes/AnimationGraphSchema.h b/Engine/Source/Editor/AnimGraph/Classes/AnimationGraphSchema.h
--- a/Engine/Source/Editor/AnimGraph/Classes/AnimationGraphSchema.h
+++ b/Engine/Source/Editor/AnimGraph/Classes/AnimationGraphSchema.h
@@ -49,7 +49,7 @@ class UAnimationGraphSchema : public UEdGraphSchema_K2
 	virtual void GetAssetsPinHoverMessage(const TArray<FAssetData>& Assets, const UEdGraphPin* HoverPin, FString& OutTooltipText, bool& OutOkIcon) const override;
 	virtual void GetAssetsGraphHoverMessage(const TArray<FAssetData>& Assets, const UEdGraph* HoverGraph, FString& OutTooltipText, bool& OutOkIcon) const override;
 	virtual void GetContextMenuActions(const UEdGraph* CurrentGraph, const UEdGraphNode* InGraphNode, const UEdGraphPin* InGraphPin, FMenuBuilder* MenuBuilder, bool bIsDebugging) const override;
-	virtual FString GetPinDisplayName(const UEdGraphPin* Pin) const override;
+	virtual FText GetPinDisplayName(const UEdGraphPin* Pin) const override;
 	virtual bool CanDuplicateGraph(UEdGraph* InSourceGraph) const override {	return false; }
 	virtual bool DoesSupportEventDispatcher() const	override {	return false; }
 	// End UEdGraphSchema interface.
diff --git a/Engine/Source/Editor/AnimGraph/Private/AnimationGraphSchema.cpp b/Engine/Source/Editor/AnimGraph/Private/AnimationGraphSchema.cpp
--- a/Engine/Source/Editor/AnimGraph/Private/AnimationGraphSchema.cpp
+++ b/Engine/Source/Editor/AnimGraph/Private/AnimationGraphSchema.cpp
@@ -490,15 +490,17 @@ void UAnimationGraphSchema::GetContextMenuActions(const UEdGraph* CurrentGraph,
 	Super::GetContextMenuActions(CurrentGraph, InGraphNode, InGraphPin, MenuBuilder, bIsDebugging);
 }
 
-FString UAnimationGraphSchema::GetPinDisplayName(const UEdGraphPin* Pin) const 
+FText UAnimationGraphSchema::GetPinDisplayName(const UEdGraphPin* Pin) const 
 {
 	check(Pin != NULL);
 
-	FString DisplayName = Super::GetPinDisplayName(Pin);
+	FText DisplayName = Super::GetPinDisplayName(Pin);
 
 	if (UAnimGraphNode_Base* Node = Cast<UAnimGraphNode_Base>(Pin->GetOwningNode()))
 	{
-		Node->PostProcessPinName(Pin, DisplayName);
+		FString ProcessedDisplayName = DisplayName.ToString();
+		Node->PostProcessPinName(Pin, ProcessedDisplayName);
+		DisplayName = FText::FromString(ProcessedDisplayName);
 	}
 
 	return DisplayName;
diff --git a/Engine/Source/Editor/BehaviorTreeEditor/Private/BehaviorTreeEditor.cpp b/Engine/Source/Editor/BehaviorTreeEditor/Private/BehaviorTreeEditor.cpp
--- a/Engine/Source/Editor/BehaviorTreeEditor/Private/BehaviorTreeEditor.cpp
+++ b/Engine/Source/Editor/BehaviorTreeEditor/Private/BehaviorTreeEditor.cpp
@@ -456,20 +456,20 @@ EVisibility FBehaviorTreeEditor::GetInjectedNodeVisibility() const
 FGraphAppearanceInfo FBehaviorTreeEditor::GetGraphAppearance() const
 {
 	FGraphAppearanceInfo AppearanceInfo;
-	AppearanceInfo.CornerText = LOCTEXT("AppearanceCornerText", "BEHAVIOR TREE").ToString();
+	AppearanceInfo.CornerText = LOCTEXT("AppearanceCornerText", "BEHAVIOR TREE");
 
 	const int32 StepIdx = Debugger.IsValid() ? Debugger->GetShownStateIndex() : 0;
 	if (Debugger.IsValid() && !Debugger->IsDebuggerRunning())
 	{
-		AppearanceInfo.PIENotifyText = TEXT("INACTIVE");
+		AppearanceInfo.PIENotifyText = LOCTEXT("InactiveLabel", "INACTIVE");
 	}
 	else if (StepIdx)
 	{
-		AppearanceInfo.PIENotifyText = FString::Printf(TEXT("%d STEP%s BACK"), StepIdx, StepIdx > 1 ? TEXT("S") : TEXT(""));
+		AppearanceInfo.PIENotifyText = FText::Format(LOCTEXT("StepsBackLabelFmt", "STEPS BACK: {0}"), FText::AsNumber(StepIdx));
 	}
 	else if (FBehaviorTreeDebugger::IsPlaySessionPaused())
 	{
-		AppearanceInfo.PIENotifyText = TEXT("PAUSED");
+		AppearanceInfo.PIENotifyText = LOCTEXT("PausedLabel", "PAUSED");
 	}
 	
 	return AppearanceInfo;
@@ -1053,9 +1053,9 @@ void FBehaviorTreeEditor::OnDebuggerActorSelected(TWeakObjectPtr<UBehaviorTreeCo
 	}
 }
 
-FString FBehaviorTreeEditor::GetDebuggerActorDesc() const
+FText FBehaviorTreeEditor::GetDebuggerActorDesc() const
 {
-	return Debugger.IsValid() ? Debugger->GetDebuggedInstanceDesc() : FString();
+	return Debugger.IsValid() ? FText::FromString(Debugger->GetDebuggedInstanceDesc()) : FText::GetEmpty();
 }
 
 void FBehaviorTreeEditor::BindDebuggerToolbarCommands()
diff --git a/Engine/Source/Editor/BehaviorTreeEditor/Private/BehaviorTreeEditor.h b/Engine/Source/Editor/BehaviorTreeEditor/Private/BehaviorTreeEditor.h
--- a/Engine/Source/Editor/BehaviorTreeEditor/Private/BehaviorTreeEditor.h
+++ b/Engine/Source/Editor/BehaviorTreeEditor/Private/BehaviorTreeEditor.h
@@ -104,7 +104,7 @@ public:
 
 	TSharedRef<class SWidget> OnGetDebuggerActorsMenu();
 	void OnDebuggerActorSelected(TWeakObjectPtr<UBehaviorTreeComponent> InstanceToDebug);
-	FString GetDebuggerActorDesc() const;
+	FText GetDebuggerActorDesc() const;
 	FGraphAppearanceInfo GetGraphAppearance() const;
 	bool InEditingMode(bool bGraphIsEditable) const;
 
diff --git a/Engine/Source/Editor/BehaviorTreeEditor/Private/BehaviorTreeGraphNode.cpp b/Engine/Source/Editor/BehaviorTreeEditor/Private/BehaviorTreeGraphNode.cpp
--- a/Engine/Source/Editor/BehaviorTreeEditor/Private/BehaviorTreeGraphNode.cpp
+++ b/Engine/Source/Editor/BehaviorTreeEditor/Private/BehaviorTreeGraphNode.cpp
@@ -311,7 +311,7 @@ void UBehaviorTreeGraphNode::DiffProperties(UStruct* Struct, void* DataA, void*
 		{
 			if(Results)
 			{
-				Diff.DisplayString = FString::Printf(*LOCTEXT("DIF_NodeProperty", "Property Changed: %s ").ToString(), *Prop->GetName());
+				Diff.DisplayString = FText::Format(LOCTEXT("DIF_NodePropertyFmt", "Property Changed: {0} "), FText::FromString(Prop->GetName()));
 				Results.Add(Diff);
 			}
 		}
@@ -324,7 +324,7 @@ void UBehaviorTreeGraphNode::FindDiffs(UEdGraphNode* OtherNode, FDiffResults& Re
 	Diff.Diff = EDiffType::NODE_PROPERTY;
 	Diff.Node1 = this;
 	Diff.Node2 = OtherNode;
-	Diff.ToolTip =  FString::Printf(*LOCTEXT("DIF_NodePropertyToolTip", "A Property of the node has changed").ToString());
+	Diff.ToolTip = LOCTEXT("DIF_NodePropertyToolTip", "A Property of the node has changed");
 	Diff.DisplayColor = FLinearColor(0.25f,0.71f,0.85f);
 	
 	UBehaviorTreeGraphNode* ThisBehaviorTreeNode = Cast<UBehaviorTreeGraphNode>(this);
diff --git a/Engine/Source/Editor/BehaviorTreeEditor/Private/DetailCustomizations/BehaviorDecoratorDetails.cpp b/Engine/Source/Editor/BehaviorTreeEditor/Private/DetailCustomizations/BehaviorDecoratorDetails.cpp
--- a/Engine/Source/Editor/BehaviorTreeEditor/Private/DetailCustomizations/BehaviorDecoratorDetails.cpp
+++ b/Engine/Source/Editor/BehaviorTreeEditor/Private/DetailCustomizations/BehaviorDecoratorDetails.cpp
@@ -130,20 +130,20 @@ TSharedRef<SWidget> FBehaviorDecoratorDetails::OnGetAbortModeContent() const
 	return MenuBuilder.MakeWidget();
 }
 
-FString FBehaviorDecoratorDetails::GetCurrentAbortModeDesc() const
+FText FBehaviorDecoratorDetails::GetCurrentAbortModeDesc() const
 {
 	uint8 ByteValue;
 	ModeProperty->GetValue(ByteValue);
 
 	for (int32 i = 0; i < ModeValues.Num(); i++)
-{
+	{
 		if (ModeValues[i].Int == ByteValue)
 		{
-			return ModeValues[i].Str;
+			return FText::FromString(ModeValues[i].Str);
 		}
 	}
 
-	return FString();
+	return FText::GetEmpty();
 }
 
 #undef LOCTEXT_NAMESPACE
diff --git a/Engine/Source/Editor/BehaviorTreeEditor/Private/DetailCustomizations/BehaviorDecoratorDetails.h b/Engine/Source/Editor/BehaviorTreeEditor/Private/DetailCustomizations/BehaviorDecoratorDetails.h
--- a/Engine/Source/Editor/BehaviorTreeEditor/Private/DetailCustomizations/BehaviorDecoratorDetails.h
+++ b/Engine/Source/Editor/BehaviorTreeEditor/Private/DetailCustomizations/BehaviorDecoratorDetails.h
@@ -26,7 +26,7 @@ private:
 
 	void OnAbortModeChange(int32 Index);
 	TSharedRef<SWidget> OnGetAbortModeContent() const;
-	FString GetCurrentAbortModeDesc() const;
+	FText GetCurrentAbortModeDesc() const;
 
 	TSharedPtr<IPropertyHandle> ModeProperty;
 
diff --git a/Engine/Source/Editor/BehaviorTreeEditor/Private/DetailCustomizations/BlackboardDecoratorDetails.cpp b/Engine/Source/Editor/BehaviorTreeEditor/Private/DetailCustomizations/BlackboardDecoratorDetails.cpp
--- a/Engine/Source/Editor/BehaviorTreeEditor/Private/DetailCustomizations/BlackboardDecoratorDetails.cpp
+++ b/Engine/Source/Editor/BehaviorTreeEditor/Private/DetailCustomizations/BlackboardDecoratorDetails.cpp
@@ -183,7 +183,7 @@ TSharedRef<SWidget> FBlackboardDecoratorDetails::OnGetEnumValueContent() const
 	return MenuBuilder.MakeWidget();
 }
 
-FString FBlackboardDecoratorDetails::GetCurrentEnumValueDesc() const
+FText FBlackboardDecoratorDetails::GetCurrentEnumValueDesc() const
 {
 	if (CachedCustomObjectType && EnumPropValues.Num() > 0)
 	{
@@ -193,11 +193,11 @@ FString FBlackboardDecoratorDetails::GetCurrentEnumValueDesc() const
 		if (Result == FPropertyAccess::Success)
 		{
 			int32 ClampedIdx = FMath::Clamp(CurrentIntValue, 0, EnumPropValues.Num());
-			return EnumPropValues[ClampedIdx];
+			return FText::FromString(EnumPropValues[ClampedIdx]);
 		}
 	}
 
-	return FString();
+	return FText::GetEmpty();
 }
 
 void FBlackboardDecoratorDetails::OnEnumValueComboChange(int32 Index)
diff --git a/Engine/Source/Editor/BehaviorTreeEditor/Private/DetailCustomizations/BlackboardDecoratorDetails.h b/Engine/Source/Editor/BehaviorTreeEditor/Private/DetailCustomizations/BlackboardDecoratorDetails.h
--- a/Engine/Source/Editor/BehaviorTreeEditor/Private/DetailCustomizations/BlackboardDecoratorDetails.h
+++ b/Engine/Source/Editor/BehaviorTreeEditor/Private/DetailCustomizations/BlackboardDecoratorDetails.h
@@ -31,7 +31,7 @@ private:
 
 	void OnEnumValueComboChange(int32 Index);
 	TSharedRef<SWidget> OnGetEnumValueContent() const;
-	FString GetCurrentEnumValueDesc() const;
+	FText GetCurrentEnumValueDesc() const;
 
 	TSharedPtr<IPropertyHandle> IntValueProperty;
 	TSharedPtr<IPropertyHandle> KeyIDProperty;
diff --git a/Engine/Source/Editor/BehaviorTreeEditor/Private/DetailCustomizations/BlackboardSelectorDetails.cpp b/Engine/Source/Editor/BehaviorTreeEditor/Private/DetailCustomizations/BlackboardSelectorDetails.cpp
--- a/Engine/Source/Editor/BehaviorTreeEditor/Private/DetailCustomizations/BlackboardSelectorDetails.cpp
+++ b/Engine/Source/Editor/BehaviorTreeEditor/Private/DetailCustomizations/BlackboardSelectorDetails.cpp
@@ -176,13 +176,13 @@ TSharedRef<SWidget> FBlackboardSelectorDetails::OnGetKeyContent() const
 	return MenuBuilder.MakeWidget();
 }
 
-FString FBlackboardSelectorDetails::GetCurrentKeyDesc() const
+FText FBlackboardSelectorDetails::GetCurrentKeyDesc() const
 {
 	FName NameValue;
 	MyKeyNameProperty->GetValue(NameValue);
 
 	const int32 KeyIdx = KeyValues.IndexOfByKey(NameValue);
-	return KeyValues.IsValidIndex(KeyIdx) ? KeyValues[KeyIdx].ToString() : NameValue.ToString();
+	return KeyValues.IsValidIndex(KeyIdx) ? FText::FromName(KeyValues[KeyIdx]) : FText::FromName(NameValue);
 }
 
 void FBlackboardSelectorDetails::OnKeyComboChange(int32 Index)
diff --git a/Engine/Source/Editor/BehaviorTreeEditor/Private/DetailCustomizations/BlackboardSelectorDetails.h b/Engine/Source/Editor/BehaviorTreeEditor/Private/DetailCustomizations/BlackboardSelectorDetails.h
--- a/Engine/Source/Editor/BehaviorTreeEditor/Private/DetailCustomizations/BlackboardSelectorDetails.h
+++ b/Engine/Source/Editor/BehaviorTreeEditor/Private/DetailCustomizations/BlackboardSelectorDetails.h
@@ -23,7 +23,7 @@ private:
 	void InitKeyFromProperty();
 	void OnKeyComboChange(int32 Index);
 	TSharedRef<SWidget> OnGetKeyContent() const;
-	FString GetCurrentKeyDesc() const;
+	FText GetCurrentKeyDesc() const;
 	bool IsEditingEnabled() const;
 
 	TSharedPtr<IPropertyHandle> MyStructProperty;
diff --git a/Engine/Source/Editor/BehaviorTreeEditor/Private/FindInBT.cpp b/Engine/Source/Editor/BehaviorTreeEditor/Private/FindInBT.cpp
--- a/Engine/Source/Editor/BehaviorTreeEditor/Private/FindInBT.cpp
+++ b/Engine/Source/Editor/BehaviorTreeEditor/Private/FindInBT.cpp
@@ -294,7 +294,7 @@ TSharedRef<ITableRow> SFindInBT::OnGenerateRow( FSearchResult InItem, const TSha
 					.Padding(2, 0)
 					[
 						SNew(STextBlock)
-						.Text(InItem->Value)
+						.Text(FText::FromString(InItem->Value))
 						.HighlightText(HighlightText)
 					]
 				]
@@ -304,7 +304,7 @@ TSharedRef<ITableRow> SFindInBT::OnGenerateRow( FSearchResult InItem, const TSha
 			.VAlign(VAlign_Center)
 			[
 				SNew(STextBlock)
-				.Text(InItem->GetNodeTypeText())
+				.Text(FText::FromString(InItem->GetNodeTypeText()))
 				.HighlightText(HighlightText)
 			]
 			+SHorizontalBox::Slot()
@@ -312,7 +312,7 @@ TSharedRef<ITableRow> SFindInBT::OnGenerateRow( FSearchResult InItem, const TSha
 			.VAlign(VAlign_Center)
 			[
 				SNew(STextBlock)
-				.Text(InItem->GetCommentText())
+				.Text(FText::FromString(InItem->GetCommentText()))
 				.ColorAndOpacity(FLinearColor::Yellow)
 				.HighlightText(HighlightText)
 			]
diff --git a/Engine/Source/Editor/BehaviorTreeEditor/Private/SBehaviorTreeBlackboardView.cpp b/Engine/Source/Editor/BehaviorTreeEditor/Private/SBehaviorTreeBlackboardView.cpp
--- a/Engine/Source/Editor/BehaviorTreeEditor/Private/SBehaviorTreeBlackboardView.cpp
+++ b/Engine/Source/Editor/BehaviorTreeEditor/Private/SBehaviorTreeBlackboardView.cpp
@@ -89,7 +89,7 @@ class SBehaviorTreeBlackboardItem : public SGraphPaletteItem
 		FSlateBrush const* IconBrush   = FEditorStyle::GetBrush(TEXT("NoBrush"));
 		GetPaletteItemIcon(GraphAction, IconBrush);
 
-		TSharedRef<SWidget> IconWidget = CreateIconWidget( GraphAction->TooltipDescription, IconBrush, FLinearColor::White );
+		TSharedRef<SWidget> IconWidget = CreateIconWidget( FText::FromString(GraphAction->TooltipDescription), IconBrush, FLinearColor::White );
 		TSharedRef<SWidget> NameSlotWidget = CreateTextSlotWidget( NameFont, InCreateData, BlackboardEntryAction->bIsInherited );
 		TSharedRef<SWidget> DebugSlotWidget = CreateDebugSlotWidget( NameFont );
 
diff --git a/Engine/Source/Editor/BehaviorTreeEditor/Private/SBehaviorTreeDiff.cpp b/Engine/Source/Editor/BehaviorTreeEditor/Private/SBehaviorTreeDiff.cpp
--- a/Engine/Source/Editor/BehaviorTreeEditor/Private/SBehaviorTreeDiff.cpp
+++ b/Engine/Source/Editor/BehaviorTreeEditor/Private/SBehaviorTreeDiff.cpp
@@ -27,13 +27,13 @@ struct FDiffResultItem: public TSharedFromThis<FDiffResultItem>
 	 */
 	TSharedRef<SWidget>	GenerateWidget() const
 	{
-		FString ToolTip = Result.ToolTip;
+		FText ToolTip = Result.ToolTip;
 		FLinearColor Color = Result.DisplayColor;
-		FString Text = Result.DisplayString;
-		if(Text.Len() == 0)
+		FText Text = Result.DisplayString;
+		if(Text.IsEmpty())
 		{
-			Text = LOCTEXT("DIF_UnknownDiff", "Unknown Diff").ToString();
-			ToolTip = LOCTEXT("DIF_Confused", "There is an unspecified difference").ToString();
+			Text = LOCTEXT("DIF_UnknownDiff", "Unknown Diff");
+			ToolTip = LOCTEXT("DIF_Confused", "There is an unspecified difference");
 		}
 		return SNew(STextBlock)
 			.ToolTipText(ToolTip)
@@ -100,7 +100,7 @@ void SBehaviorTreeDiff::Construct( const FArguments& InArgs )
 		.VAlign(VAlign_Center)
 		[
 			SNew(STextBlock)
-			.Text(LOCTEXT("BehaviorTreeDifGraphsToolTip", "Select Graph to Diff").ToString())
+			.Text(LOCTEXT("BehaviorTreeDifGraphsToolTip", "Select Graph to Diff"))
 		];
 
 	this->ChildSlot
@@ -125,7 +125,7 @@ void SBehaviorTreeDiff::Construct( const FArguments& InArgs )
 						.Content()
 						[
 							SNew(STextBlock)
-							.Text(LOCTEXT("DifBehaviorTreeDefaults", "Default Diff").ToString())
+							.Text(LOCTEXT("DifBehaviorTreeDefaults", "Default Diff"))
 						]
 					]
 					+SVerticalBox::Slot()
@@ -390,7 +390,7 @@ void SBehaviorTreeDiff::FBehaviorTreeDiffPanel::GeneratePanel(UEdGraph* Graph, U
 		.HAlign(HAlign_Center)
 		.VAlign(VAlign_Center)
 		[
-			SNew(STextBlock).Text( LOCTEXT("BTDifPanelNoGraphTip", "Graph does not exist in this revision").ToString())
+			SNew(STextBlock).Text( LOCTEXT("BTDifPanelNoGraphTip", "Graph does not exist in this revision"))
 		];
 
 	FPropertyEditorModule& PropertyEditorModule = FModuleManager::GetModuleChecked<FPropertyEditorModule>( "PropertyEditor" );
@@ -406,7 +406,7 @@ void SBehaviorTreeDiff::FBehaviorTreeDiffPanel::GeneratePanel(UEdGraph* Graph, U
 		InEvents.OnSelectionChanged = SGraphEditor::FOnSelectionChanged::CreateRaw(this, &SBehaviorTreeDiff::FBehaviorTreeDiffPanel::OnSelectionChanged);
 
 		FGraphAppearanceInfo AppearanceInfo;
-		AppearanceInfo.CornerText = LOCTEXT("AppearanceCornerText_BehaviorDif", "DIFF").ToString();
+		AppearanceInfo.CornerText = LOCTEXT("AppearanceCornerText_BehaviorDif", "DIFF");
 
 		if (!GraphEditorCommands.IsValid())
 		{
@@ -457,31 +457,38 @@ void SBehaviorTreeDiff::FBehaviorTreeDiffPanel::GeneratePanel(UEdGraph* Graph, U
 	GraphEditorBorder->SetContent(Widget.ToSharedRef());
 }
 
-FString SBehaviorTreeDiff::FBehaviorTreeDiffPanel::GetTitle() const
+FText SBehaviorTreeDiff::FBehaviorTreeDiffPanel::GetTitle() const
 {
-	FString Title = LOCTEXT("CurrentRevision", "Current Revision").ToString();
+	FText Title = LOCTEXT("CurrentRevision", "Current Revision");
 
 	// if this isn't the current working version being displayed
 	if (RevisionInfo.Revision >= 0)
 	{
-		FString DateString = RevisionInfo.Date.ToString(TEXT("%m/%d/%Y"));
+		// Don't use grouping on the revision or CL numbers to match how Perforce displays them
+		static const FNumberFormattingOptions RevisionFormatOptions = FNumberFormattingOptions()
+			.SetUseGrouping(false);
+
+		const FText DateText = FText::AsDate(RevisionInfo.Date, EDateTimeStyle::Short);
+		const FText RevisionText = FText::AsNumber(RevisionInfo.Revision, &RevisionFormatOptions);
+		const FText ChangelistText = FText::AsNumber(RevisionInfo.Changelist, &RevisionFormatOptions);
+
 		if (bShowAssetName)
 		{
 			FString AssetName = BehaviorTree->GetName();
-			FText LocalizedFormat = LOCTEXT("NamedRevisionDiff", "%s - Revision %i, CL %i, %s");
-			Title = FString::Printf(*LocalizedFormat.ToString(), *AssetName, RevisionInfo.Revision, RevisionInfo.Changelist, *DateString);
+			FText LocalizedFormat = LOCTEXT("NamedRevisionDiffFmt", "{0} - Revision {1}, CL {2}, {3}");
+			Title = FText::Format(LocalizedFormat, FText::FromString(AssetName), RevisionText, ChangelistText, DateText);
 		}
 		else
 		{
-			FText LocalizedFormat = LOCTEXT("PreviousRevisionDif", "Revision %i, CL %i, %s");
-			Title = FString::Printf(*LocalizedFormat.ToString(), RevisionInfo.Revision, RevisionInfo.Changelist, *DateString);
+			FText LocalizedFormat = LOCTEXT("PreviousRevisionDifFmt", "Revision {0}, CL {1}, {2}");
+			Title = FText::Format(LocalizedFormat, RevisionText, ChangelistText, DateText);
 		}
 	}
 	else if (bShowAssetName)
 	{
 		FString AssetName = BehaviorTree->GetName();
-		FText LocalizedFormat = LOCTEXT("NamedCurrentRevision", "%s - Current Revision");
-		Title = FString::Printf(*LocalizedFormat.ToString(), *AssetName);
+		FText LocalizedFormat = LOCTEXT("NamedCurrentRevisionFmt", "{0} - Current Revision");
+		Title = FText::Format(LocalizedFormat, FText::FromString(AssetName));
 	}
 
 	return Title;
diff --git a/Engine/Source/Editor/BehaviorTreeEditor/Private/SBehaviorTreeDiff.h b/Engine/Source/Editor/BehaviorTreeEditor/Private/SBehaviorTreeDiff.h
--- a/Engine/Source/Editor/BehaviorTreeEditor/Private/SBehaviorTreeDiff.h
+++ b/Engine/Source/Editor/BehaviorTreeEditor/Private/SBehaviorTreeDiff.h
@@ -41,7 +41,7 @@ private:
 		 * Returns title for this panel
 		 * @return The Title
 		 */
-		FString GetTitle() const;
+		FText GetTitle() const;
 
 		/** 
 		 * Called when user hits keyboard shortcut to copy nodes
diff --git a/Engine/Source/Editor/BehaviorTreeEditor/Private/SGraphNode_BehaviorTree.cpp b/Engine/Source/Editor/BehaviorTreeEditor/Private/SGraphNode_BehaviorTree.cpp
--- a/Engine/Source/Editor/BehaviorTreeEditor/Private/SGraphNode_BehaviorTree.cpp
+++ b/Engine/Source/Editor/BehaviorTreeEditor/Private/SGraphNode_BehaviorTree.cpp
@@ -878,10 +878,10 @@ FReply SGraphNode_BehaviorTree::OnDrop( const FGeometry& MyGeometry, const FDrag
 	return SGraphNode::OnDrop(MyGeometry, DragDropEvent);
 }
 
-FString	SGraphNode_BehaviorTree::GetDescription() const
+FText SGraphNode_BehaviorTree::GetDescription() const
 {
 	UBehaviorTreeGraphNode* StateNode = CastChecked<UBehaviorTreeGraphNode>(GraphNode);
-	return StateNode->GetDescription();
+	return FText::FromString(StateNode->GetDescription());
 }
 
 FString SGraphNode_BehaviorTree::GetPinTooltip(UEdGraphPin* GraphPinObj) const
@@ -982,7 +982,7 @@ TSharedPtr<SToolTip> SGraphNode_BehaviorTree::GetComplexTooltip()
 					// Create the tooltip graph preview, make sure to disable state overlays to
 					// prevent the PIE / read-only borders from obscuring the graph
 					SNew(SGraphPreviewer, DecoratorNode->GetBoundGraph())
-					.CornerOverlayText(LOCTEXT("CompositeDecoratorOverlayText", "Composite Decorator").ToString())
+					.CornerOverlayText(LOCTEXT("CompositeDecoratorOverlayText", "Composite Decorator"))
 					.ShowGraphStateOverlay(false)
 				]
 				+SOverlay::Slot()
@@ -1009,7 +1009,7 @@ TSharedPtr<SToolTip> SGraphNode_BehaviorTree::GetComplexTooltip()
 						// Create the tooltip graph preview, make sure to disable state overlays to
 						// prevent the PIE / read-only borders from obscuring the graph
 						SNew(SGraphPreviewer, RunBehavior->GetSubtreeAsset()->BTGraph)
-						.CornerOverlayText(LOCTEXT("RunBehaviorOverlayText", "Run Behavior").ToString())
+						.CornerOverlayText(LOCTEXT("RunBehaviorOverlayText", "Run Behavior"))
 						.ShowGraphStateOverlay(false)
 					]
 					+SOverlay::Slot()
diff --git a/Engine/Source/Editor/BehaviorTreeEditor/Private/SGraphNode_BehaviorTree.h b/Engine/Source/Editor/BehaviorTreeEditor/Private/SGraphNode_BehaviorTree.h
--- a/Engine/Source/Editor/BehaviorTreeEditor/Private/SGraphNode_BehaviorTree.h
+++ b/Engine/Source/Editor/BehaviorTreeEditor/Private/SGraphNode_BehaviorTree.h
@@ -105,7 +105,7 @@ protected:
 
 	FSlateColor GetBorderBackgroundColor() const;
 	FSlateColor GetBackgroundColor() const;
-	FString	GetDescription() const;
+	FText GetDescription() const;
 
 	virtual const FSlateBrush* GetNameIcon() const;
 	virtual EVisibility GetBlueprintIconVisibility() const;
diff --git a/Engine/Source/Editor/BlueprintGraph/Classes/EdGraphSchema_K2.h b/Engine/Source/Editor/BlueprintGraph/Classes/EdGraphSchema_K2.h
--- a/Engine/Source/Editor/BlueprintGraph/Classes/EdGraphSchema_K2.h
+++ b/Engine/Source/Editor/BlueprintGraph/Classes/EdGraphSchema_K2.h
@@ -388,7 +388,7 @@ public:
 	virtual bool ShouldHidePinDefaultValue(UEdGraphPin* Pin) const override;
 	virtual bool ShouldShowAssetPickerForPin(UEdGraphPin* Pin) const override;
 	virtual FLinearColor GetPinTypeColor(const FEdGraphPinType& PinType) const override;
-	virtual FString GetPinDisplayName(const UEdGraphPin* Pin) const override;
+	virtual FText GetPinDisplayName(const UEdGraphPin* Pin) const override;
 	virtual void ConstructBasicPinTooltip(const UEdGraphPin& Pin, const FText& PinDescription, FString& TooltipOut) const override;
 	virtual EGraphType GetGraphType(const UEdGraph* TestEdGraph) const override;
 	virtual bool IsTitleBarPin(const UEdGraphPin& Pin) const override;
diff --git a/Engine/Source/Editor/BlueprintGraph/Classes/K2Node_Knot.h b/Engine/Source/Editor/BlueprintGraph/Classes/K2Node_Knot.h
--- a/Engine/Source/Editor/BlueprintGraph/Classes/K2Node_Knot.h
+++ b/Engine/Source/Editor/BlueprintGraph/Classes/K2Node_Knot.h
@@ -16,7 +16,7 @@ public:
 	virtual FText GetNodeTitle(ENodeTitleType::Type TitleType) const override;
 	virtual void GetMenuEntries(struct FGraphContextMenuBuilder& ContextMenuBuilder) const override;
 	virtual bool ShouldOverridePinNames() const override;
-	virtual FString GetPinNameOverride(const UEdGraphPin& Pin) const override;
+	virtual FText GetPinNameOverride(const UEdGraphPin& Pin) const override;
 	virtual void OnRenameNode(const FString& NewName) override;
 	virtual TSharedPtr<class INameValidatorInterface> MakeNameValidator() const override;
 	virtual bool AllowSplitPins() const override;
diff --git a/Engine/Source/Editor/BlueprintGraph/Classes/K2Node_VariableSet.h b/Engine/Source/Editor/BlueprintGraph/Classes/K2Node_VariableSet.h
--- a/Engine/Source/Editor/BlueprintGraph/Classes/K2Node_VariableSet.h
+++ b/Engine/Source/Editor/BlueprintGraph/Classes/K2Node_VariableSet.h
@@ -17,7 +17,7 @@ class UK2Node_VariableSet : public UK2Node_Variable
 	virtual FText GetTooltipText() const override;
 	virtual FText GetNodeTitle(ENodeTitleType::Type TitleType) const override;
 	virtual bool ShouldOverridePinNames() const override { return true; }
-	virtual FString GetPinNameOverride(const UEdGraphPin& Pin) const override;
+	virtual FText GetPinNameOverride(const UEdGraphPin& Pin) const override;
 	// End UEdGraphNode interface
 
 	// Begin K2Node interface
diff --git a/Engine/Source/Editor/BlueprintGraph/Private/EdGraphSchema_K2.cpp b/Engine/Source/Editor/BlueprintGraph/Private/EdGraphSchema_K2.cpp
--- a/Engine/Source/Editor/BlueprintGraph/Private/EdGraphSchema_K2.cpp
+++ b/Engine/Source/Editor/BlueprintGraph/Private/EdGraphSchema_K2.cpp
@@ -2716,9 +2716,9 @@ FLinearColor UEdGraphSchema_K2::GetPinTypeColor(const FEdGraphPinType& PinType)
 	return Settings->DefaultPinTypeColor;
 }
 
-FString UEdGraphSchema_K2::GetPinDisplayName(const UEdGraphPin* Pin) const 
+FText UEdGraphSchema_K2::GetPinDisplayName(const UEdGraphPin* Pin) const 
 {
-	FString DisplayName;
+	FText DisplayName = FText::GetEmpty();
 
 	if (Pin != NULL)
 	{
@@ -2733,15 +2733,15 @@ FString UEdGraphSchema_K2::GetPinDisplayName(const UEdGraphPin* Pin) const
 	
 			// bit of a hack to hide 'execute' and 'then' pin names
 			if ((Pin->PinType.PinCategory == PC_Exec) && 
-				((DisplayName == PN_Execute) || (DisplayName == PN_Then)))
+				((DisplayName.ToString() == PN_Execute) || (DisplayName.ToString() == PN_Then)))
 			{
-				DisplayName = FString(TEXT(""));
+				DisplayName = FText::GetEmpty();
 			}
 		}
 
 		if( GEditor && GetDefault<UEditorStyleSettings>()->bShowFriendlyNames )
 		{
-			DisplayName = FName::NameToDisplayString(DisplayName, Pin->PinType.PinCategory == PC_Boolean);
+			DisplayName = FText::FromString(FName::NameToDisplayString(DisplayName.ToString(), Pin->PinType.PinCategory == PC_Boolean));
 		}
 	}
 	return DisplayName;
@@ -2763,7 +2763,7 @@ void UEdGraphSchema_K2::ConstructBasicPinTooltip(const UEdGraphPin& Pin, const F
 			UEdGraphSchema_K2 const* const K2Schema = Cast<const UEdGraphSchema_K2>(PinNode->GetSchema());
 			if (ensure(K2Schema != NULL)) // ensure that this node belongs to this schema
 			{
-				Args.Add(TEXT("DisplayName"), FText::FromString(GetPinDisplayName(&Pin)));
+				Args.Add(TEXT("DisplayName"), GetPinDisplayName(&Pin));
 				Args.Add(TEXT("LineFeed1"), FText::FromString(TEXT("\n")));
 			}
 		}
diff --git a/Engine/Source/Editor/BlueprintGraph/Private/K2Node_ConstructObjectFromClass.cpp b/Engine/Source/Editor/BlueprintGraph/Private/K2Node_ConstructObjectFromClass.cpp
--- a/Engine/Source/Editor/BlueprintGraph/Private/K2Node_ConstructObjectFromClass.cpp
+++ b/Engine/Source/Editor/BlueprintGraph/Private/K2Node_ConstructObjectFromClass.cpp
@@ -62,7 +62,7 @@ void UK2Node_ConstructObjectFromClass::SetPinToolTip(UEdGraphPin& MutatablePin,
 	if (K2Schema != nullptr)
 	{
 		MutatablePin.PinToolTip += TEXT(" ");
-		MutatablePin.PinToolTip += K2Schema->GetPinDisplayName(&MutatablePin);
+		MutatablePin.PinToolTip += K2Schema->GetPinDisplayName(&MutatablePin).ToString();
 	}
 
 	MutatablePin.PinToolTip += FString(TEXT("\n")) + PinDescription.ToString();
diff --git a/Engine/Source/Editor/BlueprintGraph/Private/K2Node_FormatText.cpp b/Engine/Source/Editor/BlueprintGraph/Private/K2Node_FormatText.cpp
--- a/Engine/Source/Editor/BlueprintGraph/Private/K2Node_FormatText.cpp
+++ b/Engine/Source/Editor/BlueprintGraph/Private/K2Node_FormatText.cpp
@@ -51,9 +51,9 @@ FText UK2Node_FormatText::GetNodeTitle(ENodeTitleType::Type TitleType) const
 	return LOCTEXT("FormatText_Title", "Format Text");
 }
 
-FString UK2Node_FormatText::GetPinDisplayName(const UEdGraphPin* Pin) const
+FText UK2Node_FormatText::GetPinDisplayName(const UEdGraphPin* Pin) const
 {
-	return Pin->PinName;
+	return FText::FromString(Pin->PinName);
 }
 
 FText UK2Node_FormatText::GetUniquePinName()
diff --git a/Engine/Source/Editor/BlueprintGraph/Private/K2Node_Knot.cpp b/Engine/Source/Editor/BlueprintGraph/Private/K2Node_Knot.cpp
--- a/Engine/Source/Editor/BlueprintGraph/Private/K2Node_Knot.cpp
+++ b/Engine/Source/Editor/BlueprintGraph/Private/K2Node_Knot.cpp
@@ -172,10 +172,10 @@ bool UK2Node_Knot::ShouldOverridePinNames() const
 	return true;
 }
 
-FString UK2Node_Knot::GetPinNameOverride(const UEdGraphPin& Pin) const
+FText UK2Node_Knot::GetPinNameOverride(const UEdGraphPin& Pin) const
 {
 	// Keep the pin size tiny
-	return FString();
+	return FText::GetEmpty();
 }
 
 void UK2Node_Knot::OnRenameNode(const FString& NewName)
diff --git a/Engine/Source/Editor/BlueprintGraph/Private/K2Node_Timeline.cpp b/Engine/Source/Editor/BlueprintGraph/Private/K2Node_Timeline.cpp
--- a/Engine/Source/Editor/BlueprintGraph/Private/K2Node_Timeline.cpp
+++ b/Engine/Source/Editor/BlueprintGraph/Private/K2Node_Timeline.cpp
@@ -385,9 +385,9 @@ void FindExactTimelineDifference(struct FDiffResults& Results, FDiffSingleResult
 		Args.Add(TEXT("NodeName"), NodeName);
 
 		Result.Diff = EDiffType::TIMELINE_NUM_TRACKS;
-		Result.ToolTip =  FText::Format(LOCTEXT("DIF_TimelineNumTracksToolTip", "The number of {TrackType} tracks in Timeline '{NodeName}' has changed"), Args).ToString();
+		Result.ToolTip =  FText::Format(LOCTEXT("DIF_TimelineNumTracksToolTip", "The number of {TrackType} tracks in Timeline '{NodeName}' has changed"), Args);
 		Result.DisplayColor = FLinearColor(0.05f,0.261f,0.775f);
-		Result.DisplayString = FText::Format(LOCTEXT("DIF_TimelineNumTracks", "{TrackType} Track Count '{NodeName}'"), Args).ToString();
+		Result.DisplayString = FText::Format(LOCTEXT("DIF_TimelineNumTracks", "{TrackType} Track Count '{NodeName}'"), Args);
 		Results.Add(Result);
 		return;
 	}
@@ -404,9 +404,9 @@ void FindExactTimelineDifference(struct FDiffResults& Results, FDiffSingleResult
 			Args.Add(TEXT("NodeName"), NodeName);
 
 			Result.Diff = EDiffType::TIMELINE_TRACK_MODIFIED;
-			Result.ToolTip =  FText::Format(LOCTEXT("DIF_TimelineTrackModifiedToolTip", "Track '{TrackName}' of Timeline '{NodeName}' was Modified"), Args).ToString();
+			Result.ToolTip =  FText::Format(LOCTEXT("DIF_TimelineTrackModifiedToolTip", "Track '{TrackName}' of Timeline '{NodeName}' was Modified"), Args);
 			Result.DisplayColor = FLinearColor(0.75f,0.1f,0.15f);
-			Result.DisplayString = FText::Format(LOCTEXT("DIF_TimelineTrackModified", "Track Modified '{TrackName}'"), Args).ToString();
+			Result.DisplayString = FText::Format(LOCTEXT("DIF_TimelineTrackModified", "Track Modified '{TrackName}'"), Args);
 			Results.Add(Result);
 			break;
 		}
@@ -440,9 +440,9 @@ void UK2Node_Timeline::FindDiffs( class UEdGraphNode* OtherNode, struct FDiffRes
 			FFormatNamedArguments Args;
 			Args.Add(TEXT("NodeName"), NodeName);
 
-			Diff.ToolTip =  FText::Format(LOCTEXT("DIF_TimelineAutoPlayToolTip", "Timeline '{NodeName}' had its AutoPlay state changed"), Args).ToString();
+			Diff.ToolTip =  FText::Format(LOCTEXT("DIF_TimelineAutoPlayToolTip", "Timeline '{NodeName}' had its AutoPlay state changed"), Args);
 			Diff.DisplayColor = FLinearColor(0.15f,0.61f,0.15f);
-			Diff.DisplayString = FText::Format(LOCTEXT("DIF_TimelineAutoPlay", "Timeline AutoPlay Changed '{NodeName}'"), Args).ToString();
+			Diff.DisplayString = FText::Format(LOCTEXT("DIF_TimelineAutoPlay", "Timeline AutoPlay Changed '{NodeName}'"), Args);
 			Results.Add(Diff);
 		}
 		if(Template1->bLoop != Template2->bLoop)
@@ -453,9 +453,9 @@ void UK2Node_Timeline::FindDiffs( class UEdGraphNode* OtherNode, struct FDiffRes
 			FFormatNamedArguments Args;
 			Args.Add(TEXT("NodeName"), NodeName);
 
-			Diff.ToolTip =  FText::Format(LOCTEXT("DIF_TimelineLoopingToolTip", "Timeline '{NodeName}' had its looping state changed"), Args).ToString();
+			Diff.ToolTip =  FText::Format(LOCTEXT("DIF_TimelineLoopingToolTip", "Timeline '{NodeName}' had its looping state changed"), Args);
 			Diff.DisplayColor = FLinearColor(0.75f,0.1f,0.75f);
-			Diff.DisplayString =  FText::Format(LOCTEXT("DIF_TimelineLooping", "Timeline Loop Changed '{NodeName}'"), Args).ToString();
+			Diff.DisplayString =  FText::Format(LOCTEXT("DIF_TimelineLooping", "Timeline Loop Changed '{NodeName}'"), Args);
 			Results.Add(Diff);
 		}
 		if(Template1->TimelineLength != Template2->TimelineLength)
@@ -468,9 +468,9 @@ void UK2Node_Timeline::FindDiffs( class UEdGraphNode* OtherNode, struct FDiffRes
 			Args.Add(TEXT("TimelineLength2"), Template2->TimelineLength);
 
 			Diff.Diff = EDiffType::TIMELINE_LENGTH;
-			Diff.ToolTip = FText::Format(LOCTEXT("DIF_TimelineLengthToolTip", "Length of Timeline '{NodeName}' has changed. Was {TimelineLength1}, but is now {TimelineLength2}"), Args).ToString();
+			Diff.ToolTip = FText::Format(LOCTEXT("DIF_TimelineLengthToolTip", "Length of Timeline '{NodeName}' has changed. Was {TimelineLength1}, but is now {TimelineLength2}"), Args);
 			Diff.DisplayColor = FLinearColor(0.25f,0.1f,0.15f);
-			Diff.DisplayString =  FText::Format(LOCTEXT("DIF_TimelineLength", "Timeline Length '{NodeName}' [{TimelineLength1} -> {TimelineLength2}]"), Args).ToString();
+			Diff.DisplayString =  FText::Format(LOCTEXT("DIF_TimelineLength", "Timeline Length '{NodeName}' [{TimelineLength1} -> {TimelineLength2}]"), Args);
 			Results.Add(Diff);
 		}
 
diff --git a/Engine/Source/Editor/BlueprintGraph/Private/K2Node_VariableSet.cpp b/Engine/Source/Editor/BlueprintGraph/Private/K2Node_VariableSet.cpp
--- a/Engine/Source/Editor/BlueprintGraph/Private/K2Node_VariableSet.cpp
+++ b/Engine/Source/Editor/BlueprintGraph/Private/K2Node_VariableSet.cpp
@@ -368,16 +368,16 @@ void UK2Node_VariableSet::CreateOutputPinTooltip()
 	Pin->PinToolTip = NSLOCTEXT("K2Node", "SetPinOutputTooltip", "Retrieves the value of the variable, can use instead of a separate Get node").ToString();
 }
 
-FString UK2Node_VariableSet::GetPinNameOverride(const UEdGraphPin& Pin) const
+FText UK2Node_VariableSet::GetPinNameOverride(const UEdGraphPin& Pin) const
 {
 	// Stop the output pin for the variable, effectively the "get" pin, from displaying a name.
 	const UEdGraphSchema_K2* K2Schema = GetDefault<UEdGraphSchema_K2>();
 	if(Pin.Direction == EGPD_Output || Pin.PinType.PinCategory == K2Schema->PC_Exec)
 	{
-		return FString();
+		return FText::GetEmpty();
 	}
 
-	return !Pin.PinFriendlyName.IsEmpty() ? Pin.PinFriendlyName.ToString() : Pin.PinName;
+	return !Pin.PinFriendlyName.IsEmpty() ? Pin.PinFriendlyName : FText::FromString(Pin.PinName);
 }
 
 void UK2Node_VariableSet::ExpandNode(class FKismetCompilerContext& CompilerContext, UEdGraph* SourceGraph)
diff --git a/Engine/Source/Editor/Blutility/Private/BlutilityDetailsPanel.cpp b/Engine/Source/Editor/Blutility/Private/BlutilityDetailsPanel.cpp
--- a/Engine/Source/Editor/Blutility/Private/BlutilityDetailsPanel.cpp
+++ b/Engine/Source/Editor/Blutility/Private/BlutilityDetailsPanel.cpp
@@ -71,7 +71,7 @@ void FEditorUtilityInstanceDetails::CustomizeDetails(IDetailLayoutBuilder& Detai
 			ActionsCategory.AddCustomRow( FText::FromString(PlacedActorCDO->HelpText) )
 			[
 				SNew(STextBlock)
-				.Text(PlacedActorCDO->HelpText)
+				.Text(FText::FromString(PlacedActorCDO->HelpText))
 			];
 		}
 		
@@ -81,7 +81,7 @@ void FEditorUtilityInstanceDetails::CustomizeDetails(IDetailLayoutBuilder& Detai
 			ActionsCategory.AddCustomRow( FText::FromString(GlobalBlutilityCDO->HelpText) )
 			[
 				SNew(STextBlock)
-				.Text(GlobalBlutilityCDO->HelpText)
+				.Text(FText::FromString(GlobalBlutilityCDO->HelpText))
 			];
 		}
 
diff --git a/Engine/Source/Editor/ClassViewer/Private/SClassViewer.cpp b/Engine/Source/Editor/ClassViewer/Private/SClassViewer.cpp
--- a/Engine/Source/Editor/ClassViewer/Private/SClassViewer.cpp
+++ b/Engine/Source/Editor/ClassViewer/Private/SClassViewer.cpp
@@ -1388,7 +1388,7 @@ public:
 				.VAlign(VAlign_Center)
 				[
 					SNew( STextBlock )
-						.Text( *ClassName.Get() )
+						.Text( FText::FromString(*ClassName.Get()) )
 						.HighlightText(*InArgs._HighlightText)
 						.ColorAndOpacity( this, &SClassItem::GetTextColor)
 						.ToolTip(Local::GetToolTip(AssociatedNode))
diff --git a/Engine/Source/Editor/ClassViewer/Public/ClassViewerModule.h b/Engine/Source/Editor/ClassViewer/Public/ClassViewerModule.h
--- a/Engine/Source/Editor/ClassViewer/Public/ClassViewerModule.h
+++ b/Engine/Source/Editor/ClassViewer/Public/ClassViewerModule.h
@@ -77,7 +77,7 @@ public:
 	bool bShowDisplayNames;
 
 	/** the title string of the class viewer if required. */
-	FString ViewerTitleString;
+	FText ViewerTitleString;
 
 	/** The property this class viewer be working on. */
 	TSharedPtr<class IPropertyHandle> PropertyHandle;
diff --git a/Engine/Source/Editor/ContentBrowser/Private/AssetViewWidgets.cpp b/Engine/Source/Editor/ContentBrowser/Private/AssetViewWidgets.cpp
--- a/Engine/Source/Editor/ContentBrowser/Private/AssetViewWidgets.cpp
+++ b/Engine/Source/Editor/ContentBrowser/Private/AssetViewWidgets.cpp
@@ -716,7 +716,7 @@ EVisibility SAssetViewItem::GetCheckedOutByOtherTextVisibility() const
 	return GetCheckedOutByOtherText().IsEmpty() ? EVisibility::Collapsed : EVisibility::Visible;
 }
 
-FString SAssetViewItem::GetCheckedOutByOtherText() const
+FText SAssetViewItem::GetCheckedOutByOtherText() const
 {
 	if ( AssetItem.IsValid() && AssetItem->GetType() != EAssetItemType::Folder && !GIsSavingPackage && !GIsGarbageCollecting )
 	{
@@ -728,12 +728,12 @@ FString SAssetViewItem::GetCheckedOutByOtherText() const
 		{
 			if ( !UserWhichHasPackageCheckedOut.IsEmpty() )
 			{
-				return SourceControlState->GetDisplayTooltip().ToString();
+				return SourceControlState->GetDisplayTooltip();
 			}
 		}
 	}
 
-	return TEXT("");
+	return FText::GetEmpty();
 }
 
 
diff --git a/Engine/Source/Editor/ContentBrowser/Private/AssetViewWidgets.h b/Engine/Source/Editor/ContentBrowser/Private/AssetViewWidgets.h
--- a/Engine/Source/Editor/ContentBrowser/Private/AssetViewWidgets.h
+++ b/Engine/Source/Editor/ContentBrowser/Private/AssetViewWidgets.h
@@ -159,7 +159,7 @@ protected:
 	EVisibility GetCheckedOutByOtherTextVisibility() const;
 
 	/** Gets the text for the checked out by other text block in the tooltip */
-	FString GetCheckedOutByOtherText() const;
+	FText GetCheckedOutByOtherText() const;
 
 	/** Helper function for CreateToolTipWidget. Gets the user description for the asset, if it exists. */
 	FText GetAssetUserDescription() const;
diff --git a/Engine/Source/Editor/ContentBrowser/Private/SAssetDialog.cpp b/Engine/Source/Editor/ContentBrowser/Private/SAssetDialog.cpp
--- a/Engine/Source/Editor/ContentBrowser/Private/SAssetDialog.cpp
+++ b/Engine/Source/Editor/ContentBrowser/Private/SAssetDialog.cpp
@@ -303,14 +303,14 @@ EVisibility SAssetDialog::GetNameErrorLabelVisibility() const
 	return GetNameErrorLabelText().IsEmpty() ? EVisibility::Hidden : EVisibility::Visible;
 }
 
-FString SAssetDialog::GetNameErrorLabelText() const
+FText SAssetDialog::GetNameErrorLabelText() const
 {
 	if (!bLastInputValidityCheckSuccessful)
 	{
-		return LastInputValidityErrorText.ToString();
+		return LastInputValidityErrorText;
 	}
 
-	return TEXT("");
+	return FText::GetEmpty();
 }
 
 void SAssetDialog::HandlePathSelected(const FString& NewPath)
diff --git a/Engine/Source/Editor/ContentBrowser/Private/SAssetDialog.h b/Engine/Source/Editor/ContentBrowser/Private/SAssetDialog.h
--- a/Engine/Source/Editor/ContentBrowser/Private/SAssetDialog.h
+++ b/Engine/Source/Editor/ContentBrowser/Private/SAssetDialog.h
@@ -46,7 +46,7 @@ private:
 	EVisibility GetNameErrorLabelVisibility() const;
 
 	/** Gets the text to display in the name error label */
-	FString GetNameErrorLabelText() const;
+	FText GetNameErrorLabelText() const;
 
 	/** Handler for when a path is selected in the path view */
 	void HandlePathSelected(const FString& NewPath);
diff --git a/Engine/Source/Editor/ContentBrowser/Private/SAssetView.cpp b/Engine/Source/Editor/ContentBrowser/Private/SAssetView.cpp
--- a/Engine/Source/Editor/ContentBrowser/Private/SAssetView.cpp
+++ b/Engine/Source/Editor/ContentBrowser/Private/SAssetView.cpp
@@ -3424,12 +3424,12 @@ FReply SAssetView::EndThumbnailEditModeClicked()
 	return FReply::Handled();
 }
 
-FString SAssetView::GetAssetCountText() const
+FText SAssetView::GetAssetCountText() const
 {
 	const int32 NumAssets = FilteredAssetItems.Num();
 	const int32 NumSelectedAssets = GetSelectedItems().Num();
 
-	FText AssetCount;
+	FText AssetCount = FText::GetEmpty();
 	if ( NumSelectedAssets == 0 )
 	{
 		if ( NumAssets == 1 )
@@ -3453,7 +3453,7 @@ FString SAssetView::GetAssetCountText() const
 		}
 	}
 
-	return AssetCount.ToString();
+	return AssetCount;
 }
 
 EVisibility SAssetView::GetEditModeLabelVisibility() const
diff --git a/Engine/Source/Editor/ContentBrowser/Private/SAssetView.h b/Engine/Source/Editor/ContentBrowser/Private/SAssetView.h
--- a/Engine/Source/Editor/ContentBrowser/Private/SAssetView.h
+++ b/Engine/Source/Editor/ContentBrowser/Private/SAssetView.h
@@ -466,7 +466,7 @@ private:
 	FReply EndThumbnailEditModeClicked();
 
 	/** Gets the text for the asset count label */
-	FString GetAssetCountText() const;
+	FText GetAssetCountText() const;
 
 	/** Gets the visibility of the Thumbnail Edit Mode label */
 	EVisibility GetEditModeLabelVisibility() const;
diff --git a/Engine/Source/Editor/DataTableEditor/Private/DataTableEditor.cpp b/Engine/Source/Editor/DataTableEditor/Private/DataTableEditor.cpp
--- a/Engine/Source/Editor/DataTableEditor/Private/DataTableEditor.cpp
+++ b/Engine/Source/Editor/DataTableEditor/Private/DataTableEditor.cpp
@@ -198,7 +198,7 @@ TSharedPtr<SUniformGridPanel> FDataTableEditor::CreateGridPanel()
 						.OnMouseButtonDown(RowClickCallback)
 						[
 							SNew(STextBlock)
-							.Text(Row[Column])
+							.Text(FText::FromString(Row[Column]))
 							.ToolTipText(bIsHeader 
 							?	(FString::Printf(TEXT("Column '%s"), *ColumnTitles[Column])) 
 							:	(FString::Printf(TEXT("%s: %s"), *ColumnTitles[Column], *Row[Column])))
diff --git a/Engine/Source/Editor/DestructibleMeshEditor/Private/DestructibleMeshEditor.cpp b/Engine/Source/Editor/DestructibleMeshEditor/Private/DestructibleMeshEditor.cpp
--- a/Engine/Source/Editor/DestructibleMeshEditor/Private/DestructibleMeshEditor.cpp
+++ b/Engine/Source/Editor/DestructibleMeshEditor/Private/DestructibleMeshEditor.cpp
@@ -445,16 +445,16 @@ void FDestructibleMeshEditor::RefreshViewport()
 	Viewport->RefreshViewport();
 }
 
-FString FDestructibleMeshEditor::HandlePreviewDepthComboBoxContent() const
+FText FDestructibleMeshEditor::HandlePreviewDepthComboBoxContent() const
 {
 	const int32 CurrentPreviewDepth = GetCurrentPreviewDepth();
 
 	if (CurrentPreviewDepth < PreviewDepths.Num())
 	{
-		return *PreviewDepths[CurrentPreviewDepth];
+		return FText::FromString(*PreviewDepths[CurrentPreviewDepth]);
 	}
 
-	return FString(TEXT("Invalid"));
+	return LOCTEXT("Invalid", "Invalid");
 }
 
 void FDestructibleMeshEditor::RegeneratePreviewDepthComboList()
@@ -494,13 +494,16 @@ void FDestructibleMeshEditor::RegeneratePreviewDepthComboList()
 
 TSharedRef<SWidget> FDestructibleMeshEditor::MakeWidgetFromString( TSharedPtr<FString> InItem )
 {
-	return SNew(STextBlock) .Text( *InItem );
+	return SNew(STextBlock) .Text( FText::FromString(*InItem) );
 }
 
-FString FDestructibleMeshEditor::GetButtonLabel() const
+FText FDestructibleMeshEditor::GetButtonLabel() const
 {
-	return FString::Printf( TEXT("%3.1f"), ExplodeFractionOfRange*ExplodeRange );
-}	
+	static const FNumberFormattingOptions FormatOptions = FNumberFormattingOptions()
+		.SetMinimumFractionalDigits(1)
+		.SetMaximumFractionalDigits(1);
+	return FText::AsNumber( ExplodeFractionOfRange*ExplodeRange, &FormatOptions );
+}
 
 void FDestructibleMeshEditor::ComboBoxSelectionChanged( TSharedPtr<FString> NewSelection )
 {
diff --git a/Engine/Source/Editor/DestructibleMeshEditor/Private/DestructibleMeshEditor.h b/Engine/Source/Editor/DestructibleMeshEditor/Private/DestructibleMeshEditor.h
--- a/Engine/Source/Editor/DestructibleMeshEditor/Private/DestructibleMeshEditor.h
+++ b/Engine/Source/Editor/DestructibleMeshEditor/Private/DestructibleMeshEditor.h
@@ -104,10 +104,10 @@ private:
 	TSharedRef<SWidget> MakeWidgetFromString( TSharedPtr<FString> InItem );
 	
 	/** Slider label. */
-	FString GetButtonLabel() const;
+	FText GetButtonLabel() const;
 
 	/** Handles display text in preview depth combo box */
-	FString HandlePreviewDepthComboBoxContent() const;
+	FText HandlePreviewDepthComboBoxContent() const;
 
 	/** A general callback for the combo boxes in the Destructible Mesh Editor to force a viewport refresh when a selection changes. */
 	void ComboBoxSelectionChanged( TSharedPtr<FString> NewSelection );
diff --git a/Engine/Source/Editor/DetailCustomizations/Private/AIDataProviderValueDetails.cpp b/Engine/Source/Editor/DetailCustomizations/Private/AIDataProviderValueDetails.cpp
--- a/Engine/Source/Editor/DetailCustomizations/Private/AIDataProviderValueDetails.cpp
+++ b/Engine/Source/Editor/DetailCustomizations/Private/AIDataProviderValueDetails.cpp
@@ -117,17 +117,17 @@ void FAIDataProviderValueDetails::OnDataFieldNameChange(int32 Index)
 	DataFieldProperty->SetValue(NameValue);
 }
 
-FString FAIDataProviderValueDetails::GetDataFieldDesc() const
+FText FAIDataProviderValueDetails::GetDataFieldDesc() const
 {
 	FName NameValue;
 	DataFieldProperty->GetValue(NameValue);
 
-	return NameValue.ToString();
+	return FText::FromString(NameValue.ToString());
 }
 
-FString FAIDataProviderValueDetails::GetValueDesc() const
+FText FAIDataProviderValueDetails::GetValueDesc() const
 {
-	return DataPtr ? DataPtr->ToString() : TEXT("empty");
+	return DataPtr ? FText::FromString(DataPtr->ToString()) : LOCTEXT("EmptyValue", "empty");
 }
 
 EVisibility FAIDataProviderValueDetails::GetDataFieldVisibility() const
diff --git a/Engine/Source/Editor/DetailCustomizations/Private/AIDataProviderValueDetails.h b/Engine/Source/Editor/DetailCustomizations/Private/AIDataProviderValueDetails.h
--- a/Engine/Source/Editor/DetailCustomizations/Private/AIDataProviderValueDetails.h
+++ b/Engine/Source/Editor/DetailCustomizations/Private/AIDataProviderValueDetails.h
@@ -26,8 +26,8 @@ private:
 	void OnBindingChanged();
 	TSharedRef<SWidget> OnGetDataFieldContent();
 	void OnDataFieldNameChange(int32 Index);
-	FString GetDataFieldDesc() const;
-	FString GetValueDesc() const;
+	FText GetDataFieldDesc() const;
+	FText GetValueDesc() const;
 	EVisibility GetBindingDescVisibility() const;
 	EVisibility GetDataFieldVisibility() const;
 	EVisibility GetDefaultValueVisibility() const;
diff --git a/Engine/Source/Editor/DetailCustomizations/Private/AnimMontageSegmentDetails.cpp b/Engine/Source/Editor/DetailCustomizations/Private/AnimMontageSegmentDetails.cpp
--- a/Engine/Source/Editor/DetailCustomizations/Private/AnimMontageSegmentDetails.cpp
+++ b/Engine/Source/Editor/DetailCustomizations/Private/AnimMontageSegmentDetails.cpp
@@ -271,26 +271,26 @@ void SAnimationSegmentViewport::Tick( const FGeometry& AllottedGeometry, const d
 
 		if (Component->IsPreviewOn())
 		{
-			Description->SetText(FText::Format( LOCTEXT("Previewing", "Previewing {0}"), FText::FromString(Component->GetPreviewText()) ).ToString());
+			Description->SetText(FText::Format( LOCTEXT("Previewing", "Previewing {0}"), FText::FromString(Component->GetPreviewText()) ));
 		}
 		else if (Component->AnimBlueprintGeneratedClass)
 		{
-			Description->SetText(FText::Format( LOCTEXT("Previewing", "Previewing {0}"), FText::FromString(Component->AnimBlueprintGeneratedClass->GetName()) ).ToString());
+			Description->SetText(FText::Format( LOCTEXT("Previewing", "Previewing {0}"), FText::FromString(Component->AnimBlueprintGeneratedClass->GetName()) ));
 		}
 		else if (Component->SkeletalMesh == NULL)
 		{
-			Description->SetText(FText::Format( LOCTEXT("NoMeshFound", "No skeletal mesh found for skeleton '{0}'"), FText::FromString(TargetSkeletonName) ).ToString());
+			Description->SetText(FText::Format( LOCTEXT("NoMeshFound", "No skeletal mesh found for skeleton '{0}'"), FText::FromString(TargetSkeletonName) ));
 		}
 		else
 		{
-			Description->SetText(LOCTEXT("Default", "Default").ToString());
+			Description->SetText(LOCTEXT("Default", "Default"));
 		}
 
 		Component->GetScene()->GetWorld()->Tick(LEVELTICK_All, InDeltaTime);
 	}
 	else
 	{
-		Description->SetText(FText::Format( LOCTEXT("NoMeshFound", "No skeletal mesh found for skeleton '{0}'"), FText::FromString(TargetSkeletonName) ).ToString());
+		Description->SetText(FText::Format( LOCTEXT("NoMeshFound", "No skeletal mesh found for skeleton '{0}'"), FText::FromString(TargetSkeletonName) ));
 	}
 }
 
diff --git a/Engine/Source/Editor/DetailCustomizations/Private/AnimSequenceDetails.cpp b/Engine/Source/Editor/DetailCustomizations/Private/AnimSequenceDetails.cpp
--- a/Engine/Source/Editor/DetailCustomizations/Private/AnimSequenceDetails.cpp
+++ b/Engine/Source/Editor/DetailCustomizations/Private/AnimSequenceDetails.cpp
@@ -186,7 +186,7 @@ EVisibility FAnimSequenceDetails::ShouldShowRefFrameIndex() const
 
 TSharedRef<SWidget> FAnimSequenceDetails::MakeRetargetSourceComboWidget( TSharedPtr<FString> InItem )
 {
-	return SNew(STextBlock) .Text( *InItem ) .Font( IDetailLayoutBuilder::GetDetailFont() );
+	return SNew(STextBlock) .Text( FText::FromString(*InItem) ) .Font( IDetailLayoutBuilder::GetDetailFont() );
 }
 
 void FAnimSequenceDetails::DelegateRetargetSourceChanged()
@@ -254,20 +254,20 @@ void FAnimSequenceDetails::OnRetargetSourceChanged( TSharedPtr<FString> NewSelec
 	}
 }
 
-FString FAnimSequenceDetails::GetRetargetSourceComboBoxContent() const
+FText FAnimSequenceDetails::GetRetargetSourceComboBoxContent() const
 {
 	FName RetargetSourceName;
 	if (RetargetSourceNameHandler->GetValue(RetargetSourceName) == FPropertyAccess::Result::MultipleValues)
 	{
-		return TEXT("Multiple Values");
+		return LOCTEXT("MultipleValues", "Multiple Values");
 	}
 
-	return (*GetRetargetSourceString(RetargetSourceName).Get());
+	return FText::FromString(*GetRetargetSourceString(RetargetSourceName).Get());
 }
 
-FString FAnimSequenceDetails::GetRetargetSourceComboBoxToolTip() const
+FText FAnimSequenceDetails::GetRetargetSourceComboBoxToolTip() const
 {
-	return TEXT("When retargeting, this pose will be used as a base of animation");
+	return LOCTEXT("RetargetSourceComboToolTip", "When retargeting, this pose will be used as a base of animation");
 }
 
 TSharedPtr<FString> FAnimSequenceDetails::GetRetargetSourceString(FName RetargetSourceName) const
@@ -474,30 +474,30 @@ void SAnimationRefPoseViewport::Tick( const FGeometry& AllottedGeometry, const d
 				}
 			}
 
-			Description->SetText( FText::Format( LOCTEXT( "Previewing", "Previewing {0}" ), FText::FromString( Component->GetPreviewText() ) ).ToString() );
+			Description->SetText( FText::Format( LOCTEXT( "Previewing", "Previewing {0}" ), FText::FromString( Component->GetPreviewText() ) ) );
 		}
 		else if ( Component->AnimBlueprintGeneratedClass )
 		{
-			Description->SetText( FText::Format( LOCTEXT( "Previewing", "Previewing {0}" ), FText::FromString( Component->AnimBlueprintGeneratedClass->GetName() ) ).ToString() );
+			Description->SetText( FText::Format( LOCTEXT( "Previewing", "Previewing {0}" ), FText::FromString( Component->AnimBlueprintGeneratedClass->GetName() ) ) );
 		}
 		else if ( AnimRef && AnimRef->GetSkeleton() != TargetSkeleton )
 		{
-			Description->SetText( FText::Format( LOCTEXT( "IncorrectSkeleton", "The preview asset doesn't work for the skeleton '{0}'" ), FText::FromString( TargetSkeletonName ) ).ToString() );
+			Description->SetText( FText::Format( LOCTEXT( "IncorrectSkeleton", "The preview asset doesn't work for the skeleton '{0}'" ), FText::FromString( TargetSkeletonName ) ) );
 		}
 		else if ( Component->SkeletalMesh == NULL )
 		{
-			Description->SetText( FText::Format( LOCTEXT( "NoMeshFound", "No skeletal mesh found for skeleton '{0}'" ), FText::FromString( TargetSkeletonName ) ).ToString() );
+			Description->SetText( FText::Format( LOCTEXT( "NoMeshFound", "No skeletal mesh found for skeleton '{0}'" ), FText::FromString( TargetSkeletonName ) ) );
 		}
 		else
 		{
-			Description->SetText( FText::Format( LOCTEXT( "SelectAnimation", "Select animation that works for skeleton '{0}'" ), FText::FromString( TargetSkeletonName ) ).ToString() );
+			Description->SetText( FText::Format( LOCTEXT( "SelectAnimation", "Select animation that works for skeleton '{0}'" ), FText::FromString( TargetSkeletonName ) ) );
 		}
 
 		Component->GetScene()->GetWorld()->Tick( LEVELTICK_All, InDeltaTime );
 	}
 	else
 	{
-		Description->SetText( FText::Format( LOCTEXT( "NoMeshFound", "No skeletal mesh found for skeleton '{0}'" ), FText::FromString( TargetSkeletonName ) ).ToString() );
+		Description->SetText( FText::Format( LOCTEXT( "NoMeshFound", "No skeletal mesh found for skeleton '{0}'" ), FText::FromString( TargetSkeletonName ) ) );
 	}
 }
 
diff --git a/Engine/Source/Editor/DetailCustomizations/Private/AnimSequenceDetails.h b/Engine/Source/Editor/DetailCustomizations/Private/AnimSequenceDetails.h
--- a/Engine/Source/Editor/DetailCustomizations/Private/AnimSequenceDetails.h
+++ b/Engine/Source/Editor/DetailCustomizations/Private/AnimSequenceDetails.h
@@ -40,8 +40,8 @@ private:
 
 	TSharedRef<SWidget> MakeRetargetSourceComboWidget( TSharedPtr<FString> InItem );
 	void OnRetargetSourceChanged( TSharedPtr<FString> NewSelection, ESelectInfo::Type SelectInfo  );
-	FString GetRetargetSourceComboBoxContent() const;
-	FString GetRetargetSourceComboBoxToolTip() const;
+	FText GetRetargetSourceComboBoxContent() const;
+	FText GetRetargetSourceComboBoxToolTip() const;
 	void OnRetargetSourceComboOpening();
 	TSharedPtr<FString> GetRetargetSourceString(FName RetargetSourceName) const;
 
