GraphEditor: Refactored connection drawing policy API to be more flexible, passing a parameter structure around instead of loose values (FConnectionParams)
Remove duplicated implementations of DrawConnection in subclasses that draw non-curved splines, overriding ComputeSplineTangent instead
Upgrade note: Any editor plugins that contain a FConnectionDrawingPolicy subclass will need to update calls to DrawSplineWithArrow, DrawConnection, and DetermineWiringStyle

diff --git a/Engine/Source/Editor/BehaviorTreeEditor/Private/BehaviorTreeConnectionDrawingPolicy.cpp b/Engine/Source/Editor/BehaviorTreeEditor/Private/BehaviorTreeConnectionDrawingPolicy.cpp
--- a/Engine/Source/Editor/BehaviorTreeEditor/Private/BehaviorTreeConnectionDrawingPolicy.cpp
+++ b/Engine/Source/Editor/BehaviorTreeEditor/Private/BehaviorTreeConnectionDrawingPolicy.cpp
@@ -12,17 +12,16 @@ FBehaviorTreeConnectionDrawingPolicy::FBehaviorTreeConnectionDrawingPolicy( int3
 {
 }
 
-void FBehaviorTreeConnectionDrawingPolicy::DetermineWiringStyle(UEdGraphPin* OutputPin, UEdGraphPin* InputPin, /*inout*/ float& Thickness, /*inout*/ FLinearColor& WireColor, /*inout*/ bool& bDrawBubbles, /*inout*/ bool& bBidirectional)
+void FBehaviorTreeConnectionDrawingPolicy::DetermineWiringStyle(UEdGraphPin* OutputPin, UEdGraphPin* InputPin, /*inout*/ FConnectionParams& Params)
 {
-	Thickness = 1.5f;
+	Params.WireThickness = 1.5f;
 
-	WireColor = BehaviorTreeColors::Connection::Default;
-	bBidirectional = false;
+	Params.WireColor = BehaviorTreeColors::Connection::Default;
 
 	const bool bDeemphasizeUnhoveredPins = HoveredPins.Num() > 0;
 	if (bDeemphasizeUnhoveredPins)
 	{
-		ApplyHoverDeemphasis(OutputPin, InputPin, /*inout*/ Thickness, /*inout*/ WireColor);
+		ApplyHoverDeemphasis(OutputPin, InputPin, /*inout*/ Params.WireThickness, /*inout*/ Params.WireColor);
 	}
 
 	UBehaviorTreeGraphNode* FromNode = OutputPin ? Cast<UBehaviorTreeGraphNode>(OutputPin->GetOwningNode()) : NULL;
@@ -32,8 +31,8 @@ void FBehaviorTreeConnectionDrawingPolicy::DetermineWiringStyle(UEdGraphPin* Out
 		if ((ToNode->bDebuggerMarkCurrentlyActive && FromNode->bDebuggerMarkCurrentlyActive) ||
 			(ToNode->bDebuggerMarkPreviouslyActive && FromNode->bDebuggerMarkPreviouslyActive))
 		{
-			Thickness = 10.0f;
-			bDrawBubbles = true;
+			Params.WireThickness = 10.0f;
+			Params.bDrawBubbles = true;
 		}
 		else if (FBehaviorTreeDebugger::IsPlaySessionPaused())
 		{
@@ -55,12 +54,12 @@ void FBehaviorTreeConnectionDrawingPolicy::DetermineWiringStyle(UEdGraphPin* Out
 
 			if (FirstToNode->bDebuggerMarkSearchSucceeded || FirstToNode->bDebuggerMarkSearchFailed)
 			{
-				Thickness = 5.0f;
-				WireColor = FirstToNode->bDebuggerMarkSearchSucceeded ? BehaviorTreeColors::Debugger::SearchSucceeded :
-					BehaviorTreeColors::Debugger::SearchFailed;
+				Params.WireThickness = 5.0f;
+				Params.WireColor = FirstToNode->bDebuggerMarkSearchSucceeded ? BehaviorTreeColors::Debugger::SearchSucceeded :
+				BehaviorTreeColors::Debugger::SearchFailed;
 
-				// hacky: use bBidirectional flag to reverse direction of connection (used by debugger)
-				bBidirectional = true;
+				// Use the bUserFlag1 flag to indicate that we need to reverse the direction of connection (used by debugger)
+				Params.bUserFlag1 = true;
 			}
 		}
 	}
@@ -83,34 +82,29 @@ void FBehaviorTreeConnectionDrawingPolicy::Draw(TMap<TSharedRef<SWidget>, FArran
 
 void FBehaviorTreeConnectionDrawingPolicy::DrawPreviewConnector(const FGeometry& PinGeometry, const FVector2D& StartPoint, const FVector2D& EndPoint, UEdGraphPin* Pin)
 {
-	float Thickness = 1.0f;
-	FLinearColor WireColor = FLinearColor::White;
-	bool bDrawBubbles = false;
-	bool bBiDirectional = false;
-	DetermineWiringStyle(Pin, NULL, /*inout*/ Thickness, /*inout*/ WireColor, /*inout*/ bDrawBubbles, /*inout*/ bBiDirectional);
+	FConnectionParams Params;
+	DetermineWiringStyle(Pin, NULL, /*inout*/ Params);
 
 	if (Pin->Direction == EEdGraphPinDirection::EGPD_Output)
 	{
-		DrawSplineWithArrow(FGeometryHelper::FindClosestPointOnGeom(PinGeometry, EndPoint), EndPoint, WireColor, Thickness, bDrawBubbles, bBiDirectional);
+		DrawSplineWithArrow(FGeometryHelper::FindClosestPointOnGeom(PinGeometry, EndPoint), EndPoint, Params);
 	}
 	else
 	{
-		DrawSplineWithArrow(FGeometryHelper::FindClosestPointOnGeom(PinGeometry, StartPoint), StartPoint, WireColor, Thickness, bDrawBubbles, bBiDirectional);
+		DrawSplineWithArrow(FGeometryHelper::FindClosestPointOnGeom(PinGeometry, StartPoint), StartPoint, Params);
 	}
-
 }
 
-
-void FBehaviorTreeConnectionDrawingPolicy::DrawSplineWithArrow(const FVector2D& StartAnchorPoint, const FVector2D& EndAnchorPoint, const FLinearColor& WireColor, float WireThickness, bool bDrawBubbles, bool Bidirectional)
+void FBehaviorTreeConnectionDrawingPolicy::DrawSplineWithArrow(const FVector2D& StartAnchorPoint, const FVector2D& EndAnchorPoint, const FConnectionParams& Params)
 {
-	// hacky: use bBidirectional flag to reverse direction of connection (used by debugger)
-	const FVector2D& P0 = Bidirectional ? EndAnchorPoint : StartAnchorPoint;
-	const FVector2D& P1 = Bidirectional ? StartAnchorPoint : EndAnchorPoint;
+	// bUserFlag1 indicates that we need to reverse the direction of connection (used by debugger)
+	const FVector2D& P0 = Params.bUserFlag1 ? EndAnchorPoint : StartAnchorPoint;
+	const FVector2D& P1 = Params.bUserFlag1 ? StartAnchorPoint : EndAnchorPoint;
 
-	Internal_DrawLineWithArrow(P0, P1, WireColor, WireThickness, bDrawBubbles);
+	Internal_DrawLineWithArrow(P0, P1, Params);
 }
 
-void FBehaviorTreeConnectionDrawingPolicy::Internal_DrawLineWithArrow(const FVector2D& StartAnchorPoint, const FVector2D& EndAnchorPoint, const FLinearColor& WireColor, float WireThickness, bool bDrawBubbles)
+void FBehaviorTreeConnectionDrawingPolicy::Internal_DrawLineWithArrow(const FVector2D& StartAnchorPoint, const FVector2D& EndAnchorPoint, const FConnectionParams& Params)
 {
 	//@TODO: Should this be scaled by zoom factor?
 	const float LineSeparationAmount = 4.5f;
@@ -126,7 +120,7 @@ void FBehaviorTreeConnectionDrawingPolicy::Internal_DrawLineWithArrow(const FVec
 	const FVector2D EndPoint = EndAnchorPoint + DirectionBias - LengthBias;
 
 	// Draw a line/spline
-	DrawConnection(WireLayerID, StartPoint, EndPoint, WireColor, WireThickness, bDrawBubbles);
+	DrawConnection(WireLayerID, StartPoint, EndPoint, Params);
 
 	// Draw the arrow
 	const FVector2D ArrowDrawPos = EndPoint - ArrowRadius;
@@ -142,11 +136,11 @@ void FBehaviorTreeConnectionDrawingPolicy::Internal_DrawLineWithArrow(const FVec
 		AngleInRadians,
 		TOptional<FVector2D>(),
 		FSlateDrawElement::RelativeToElement,
-		WireColor
+		Params.WireColor
 		);
 }
 
-void FBehaviorTreeConnectionDrawingPolicy::DrawSplineWithArrow(FGeometry& StartGeom, FGeometry& EndGeom, const FLinearColor& WireColor, float WireThickness, bool bDrawBubbles, bool Bidirectional)
+void FBehaviorTreeConnectionDrawingPolicy::DrawSplineWithArrow(const FGeometry& StartGeom, const FGeometry& EndGeom, const FConnectionParams& Params)
 {
 	// Get a reasonable seed point (halfway between the boxes)
 	const FVector2D StartCenter = FGeometryHelper::CenterOf(StartGeom);
@@ -157,65 +151,13 @@ void FBehaviorTreeConnectionDrawingPolicy::DrawSplineWithArrow(FGeometry& StartG
 	const FVector2D StartAnchorPoint = FGeometryHelper::FindClosestPointOnGeom(StartGeom, SeedPoint);
 	const FVector2D EndAnchorPoint = FGeometryHelper::FindClosestPointOnGeom(EndGeom, SeedPoint);
 
-	DrawSplineWithArrow(StartAnchorPoint, EndAnchorPoint, WireColor, WireThickness, bDrawBubbles, Bidirectional);
+	DrawSplineWithArrow(StartAnchorPoint, EndAnchorPoint, Params);
 }
 
-void FBehaviorTreeConnectionDrawingPolicy::DrawConnection(int32 LayerId, const FVector2D& Start, const FVector2D& End, const FLinearColor& InColor, float Thickness, bool bDrawBubbles)
+FVector2D FBehaviorTreeConnectionDrawingPolicy::ComputeSplineTangent(const FVector2D& Start, const FVector2D& End) const
 {
-	const FVector2D& P0 = Start;
-	const FVector2D& P1 = End;
-
-	const FVector2D Delta = End-Start;
+	const FVector2D Delta = End - Start;
 	const FVector2D NormDelta = Delta.GetSafeNormal();
 
-	const FVector2D P0Tangent = NormDelta;
-	const FVector2D P1Tangent = NormDelta;
-
-	// Draw the spline itself
-	FSlateDrawElement::MakeDrawSpaceSpline(
-		DrawElementsList,
-		LayerId,
-		P0, P0Tangent,
-		P1, P1Tangent,
-		ClippingRect,
-		Thickness,
-		ESlateDrawEffect::None,
-		InColor
-		);
-
-	if (bDrawBubbles)
-	{
-		// This table maps distance along curve to alpha
-		FInterpCurve<float> SplineReparamTable;
-		float SplineLength = MakeSplineReparamTable(P0, P0Tangent, P1, P1Tangent, SplineReparamTable);
-
-		// Draw bubbles on the spline
-		const float BubbleSpacing = 64.f * ZoomFactor;
-		const float BubbleSpeed = 192.f * ZoomFactor;
-		const FVector2D BubbleSize = BubbleImage->ImageSize * ZoomFactor * 0.1f * Thickness;
-
-		float Time = (FPlatformTime::Seconds() - GStartTime);
-		const float BubbleOffset = FMath::Fmod(Time * BubbleSpeed, BubbleSpacing);
-		const int32 NumBubbles = FMath::CeilToInt(SplineLength/BubbleSpacing);
-		for (int32 i = 0; i < NumBubbles; ++i)
-		{
-			const float Distance = ((float)i * BubbleSpacing) + BubbleOffset;
-			if (Distance < SplineLength)
-			{
-				const float Alpha = SplineReparamTable.Eval(Distance, 0.f);
-				FVector2D BubblePos = FMath::CubicInterp(P0, P0Tangent, P1, P1Tangent, Alpha);
-				BubblePos -= (BubbleSize * 0.5f);
-
-				FSlateDrawElement::MakeBox(
-					DrawElementsList,
-					LayerId,
-					FPaintGeometry( BubblePos, BubbleSize, ZoomFactor  ),
-					BubbleImage,
-					ClippingRect,
-					ESlateDrawEffect::None,
-					InColor
-					);
-			}
-		}
-	}
+	return NormDelta;
 }
diff --git a/Engine/Source/Editor/BehaviorTreeEditor/Private/BehaviorTreeConnectionDrawingPolicy.h b/Engine/Source/Editor/BehaviorTreeEditor/Private/BehaviorTreeConnectionDrawingPolicy.h
--- a/Engine/Source/Editor/BehaviorTreeEditor/Private/BehaviorTreeConnectionDrawingPolicy.h
+++ b/Engine/Source/Editor/BehaviorTreeEditor/Private/BehaviorTreeConnectionDrawingPolicy.h
@@ -14,14 +14,14 @@ public:
 	FBehaviorTreeConnectionDrawingPolicy(int32 InBackLayerID, int32 InFrontLayerID, float ZoomFactor, const FSlateRect& InClippingRect, FSlateWindowElementList& InDrawElements, UEdGraph* InGraphObj);
 
 	// FConnectionDrawingPolicy interface 
-	virtual void DetermineWiringStyle(UEdGraphPin* OutputPin, UEdGraphPin* InputPin, /*inout*/ float& Thickness, /*inout*/ FLinearColor& WireColor, /*inout*/ bool& bDrawBubbles, /*inout*/ bool& bBidirectional) override;
+	virtual void DetermineWiringStyle(UEdGraphPin* OutputPin, UEdGraphPin* InputPin, /*inout*/ FConnectionParams& Params) override;
 	virtual void Draw(TMap<TSharedRef<SWidget>, FArrangedWidget>& PinGeometries, FArrangedChildren& ArrangedNodes) override;
-	virtual void DrawSplineWithArrow(FGeometry& StartGeom, FGeometry& EndGeom, const FLinearColor& WireColor, float WireThickness, bool bDrawBubbles, bool Bidirectional) override;
-	virtual void DrawSplineWithArrow(const FVector2D& StartPoint, const FVector2D& EndPoint, const FLinearColor& WireColor, float WireThickness, bool bDrawBubbles, bool Bidirectional) override;
+	virtual void DrawSplineWithArrow(const FGeometry& StartGeom, const FGeometry& EndGeom, const FConnectionParams& Params) override;
+	virtual void DrawSplineWithArrow(const FVector2D& StartPoint, const FVector2D& EndPoint, const FConnectionParams& Params) override;
 	virtual void DrawPreviewConnector(const FGeometry& PinGeometry, const FVector2D& StartPoint, const FVector2D& EndPoint, UEdGraphPin* Pin) override;
-	virtual void DrawConnection(int32 LayerId, const FVector2D& Start, const FVector2D& End, const FLinearColor& InColor, float Thickness, bool bDrawBubbles) override;
+	virtual FVector2D ComputeSplineTangent(const FVector2D& Start, const FVector2D& End) const override;
 	// End of FConnectionDrawingPolicy interface
 
 protected:
-	void Internal_DrawLineWithArrow(const FVector2D& StartAnchorPoint, const FVector2D& EndAnchorPoint, const FLinearColor& WireColor, float WireThickness, bool bDrawBubbles);
+	void Internal_DrawLineWithArrow(const FVector2D& StartAnchorPoint, const FVector2D& EndAnchorPoint, const FConnectionParams& Params);
 };
diff --git a/Engine/Source/Editor/EnvironmentQueryEditor/Private/EnvironmentQueryConnectionDrawingPolicy.cpp b/Engine/Source/Editor/EnvironmentQueryEditor/Private/EnvironmentQueryConnectionDrawingPolicy.cpp
--- a/Engine/Source/Editor/EnvironmentQueryEditor/Private/EnvironmentQueryConnectionDrawingPolicy.cpp
+++ b/Engine/Source/Editor/EnvironmentQueryEditor/Private/EnvironmentQueryConnectionDrawingPolicy.cpp
@@ -12,17 +12,14 @@ FEnvironmentQueryConnectionDrawingPolicy::FEnvironmentQueryConnectionDrawingPoli
 {
 }
 
-void FEnvironmentQueryConnectionDrawingPolicy::DetermineWiringStyle(UEdGraphPin* OutputPin, UEdGraphPin* InputPin, /*inout*/ float& Thickness, /*inout*/ FLinearColor& WireColor, /*inout*/ bool& bDrawBubbles, /*inout*/ bool& bBidirectional)
+void FEnvironmentQueryConnectionDrawingPolicy::DetermineWiringStyle(UEdGraphPin* OutputPin, UEdGraphPin* InputPin, /*inout*/ FConnectionParams& Params)
 {
-	Thickness = 1.5f;
-
-	WireColor = FLinearColor::White;
-	bBidirectional = false;
+	Params.WireThickness = 1.5f;
 
 	const bool bDeemphasizeUnhoveredPins = HoveredPins.Num() > 0;
 	if (bDeemphasizeUnhoveredPins)
 	{
-		ApplyHoverDeemphasis(OutputPin, InputPin, /*inout*/ Thickness, /*inout*/ WireColor);
+		ApplyHoverDeemphasis(OutputPin, InputPin, /*inout*/ Params.WireThickness, /*inout*/ Params.WireColor);
 	}
 }
 
@@ -43,30 +40,25 @@ void FEnvironmentQueryConnectionDrawingPolicy::Draw(TMap<TSharedRef<SWidget>, FA
 
 void FEnvironmentQueryConnectionDrawingPolicy::DrawPreviewConnector(const FGeometry& PinGeometry, const FVector2D& StartPoint, const FVector2D& EndPoint, UEdGraphPin* Pin)
 {
-	float Thickness = 1.0f;
-	FLinearColor WireColor = FLinearColor::White;
-	bool bDrawBubbles = false;
-	bool bBiDirectional = false;
-	DetermineWiringStyle(Pin, NULL, /*inout*/ Thickness, /*inout*/ WireColor, /*inout*/ bDrawBubbles, /*inout*/ bBiDirectional);
+	FConnectionParams Params;
+	DetermineWiringStyle(Pin, NULL, /*inout*/ Params);
 
 	if (Pin->Direction == EEdGraphPinDirection::EGPD_Output)
 	{
-		DrawSplineWithArrow(FGeometryHelper::FindClosestPointOnGeom(PinGeometry, EndPoint), EndPoint, WireColor, Thickness, bDrawBubbles, bBiDirectional);
+		DrawSplineWithArrow(FGeometryHelper::FindClosestPointOnGeom(PinGeometry, EndPoint), EndPoint, Params);
 	}
 	else
 	{
-		DrawSplineWithArrow(FGeometryHelper::FindClosestPointOnGeom(PinGeometry, StartPoint), StartPoint, WireColor, Thickness, bDrawBubbles, bBiDirectional);
+		DrawSplineWithArrow(FGeometryHelper::FindClosestPointOnGeom(PinGeometry, StartPoint), StartPoint, Params);
 	}
-
 }
 
-
-void FEnvironmentQueryConnectionDrawingPolicy::DrawSplineWithArrow(const FVector2D& StartAnchorPoint, const FVector2D& EndAnchorPoint, const FLinearColor& WireColor, float WireThickness, bool bDrawBubbles, bool Bidirectional)
+void FEnvironmentQueryConnectionDrawingPolicy::DrawSplineWithArrow(const FVector2D& StartAnchorPoint, const FVector2D& EndAnchorPoint, const FConnectionParams& Params)
 {
-	Internal_DrawLineWithArrow(StartAnchorPoint, EndAnchorPoint, WireColor, WireThickness, bDrawBubbles);
+	Internal_DrawLineWithArrow(StartAnchorPoint, EndAnchorPoint, Params);
 }
 
-void FEnvironmentQueryConnectionDrawingPolicy::Internal_DrawLineWithArrow(const FVector2D& StartAnchorPoint, const FVector2D& EndAnchorPoint, const FLinearColor& WireColor, float WireThickness, bool bDrawBubbles)
+void FEnvironmentQueryConnectionDrawingPolicy::Internal_DrawLineWithArrow(const FVector2D& StartAnchorPoint, const FVector2D& EndAnchorPoint, const FConnectionParams& Params)
 {
 	//@TODO: Should this be scaled by zoom factor?
 	const float LineSeparationAmount = 4.5f;
@@ -82,7 +74,7 @@ void FEnvironmentQueryConnectionDrawingPolicy::Internal_DrawLineWithArrow(const
 	const FVector2D EndPoint = EndAnchorPoint + DirectionBias - LengthBias;
 
 	// Draw a line/spline
-	DrawConnection(WireLayerID, StartPoint, EndPoint, WireColor, WireThickness, bDrawBubbles);
+	DrawConnection(WireLayerID, StartPoint, EndPoint, Params);
 
 	// Draw the arrow
 	const FVector2D ArrowDrawPos = EndPoint - ArrowRadius;
@@ -98,11 +90,11 @@ void FEnvironmentQueryConnectionDrawingPolicy::Internal_DrawLineWithArrow(const
 		AngleInRadians,
 		TOptional<FVector2D>(),
 		FSlateDrawElement::RelativeToElement,
-		WireColor
+		Params.WireColor
 		);
 }
 
-void FEnvironmentQueryConnectionDrawingPolicy::DrawSplineWithArrow(FGeometry& StartGeom, FGeometry& EndGeom, const FLinearColor& WireColor, float WireThickness, bool bDrawBubbles, bool Bidirectional)
+void FEnvironmentQueryConnectionDrawingPolicy::DrawSplineWithArrow(const FGeometry& StartGeom, const FGeometry& EndGeom, const FConnectionParams& Params)
 {
 	// Get a reasonable seed point (halfway between the boxes)
 	const FVector2D StartCenter = FGeometryHelper::CenterOf(StartGeom);
@@ -113,31 +105,13 @@ void FEnvironmentQueryConnectionDrawingPolicy::DrawSplineWithArrow(FGeometry& St
 	const FVector2D StartAnchorPoint = FGeometryHelper::FindClosestPointOnGeom(StartGeom, SeedPoint);
 	const FVector2D EndAnchorPoint = FGeometryHelper::FindClosestPointOnGeom(EndGeom, SeedPoint);
 
-	DrawSplineWithArrow(StartAnchorPoint, EndAnchorPoint, WireColor, WireThickness, bDrawBubbles, Bidirectional);
+	DrawSplineWithArrow(StartAnchorPoint, EndAnchorPoint, Params);
 }
 
-void FEnvironmentQueryConnectionDrawingPolicy::DrawConnection(int32 LayerId, const FVector2D& Start, const FVector2D& End, const FLinearColor& InColor, float Thickness, bool bDrawBubbles)
+FVector2D FEnvironmentQueryConnectionDrawingPolicy::ComputeSplineTangent(const FVector2D& Start, const FVector2D& End) const
 {
-	const FVector2D& P0 = Start;
-	const FVector2D& P1 = End;
-
-	const FVector2D Delta = End-Start;
+	const FVector2D Delta = End - Start;
 	const FVector2D NormDelta = Delta.GetSafeNormal();
 
-	const FVector2D P0Tangent = NormDelta;
-	const FVector2D P1Tangent = NormDelta;
-
-	// Draw the spline itself
-	FSlateDrawElement::MakeDrawSpaceSpline(
-		DrawElementsList,
-		LayerId,
-		P0, P0Tangent,
-		P1, P1Tangent,
-		ClippingRect,
-		Thickness,
-		ESlateDrawEffect::None,
-		InColor
-		);
-
-	//@TODO: Handle bDrawBubbles
+	return NormDelta;
 }
diff --git a/Engine/Source/Editor/EnvironmentQueryEditor/Private/EnvironmentQueryConnectionDrawingPolicy.h b/Engine/Source/Editor/EnvironmentQueryEditor/Private/EnvironmentQueryConnectionDrawingPolicy.h
--- a/Engine/Source/Editor/EnvironmentQueryEditor/Private/EnvironmentQueryConnectionDrawingPolicy.h
+++ b/Engine/Source/Editor/EnvironmentQueryEditor/Private/EnvironmentQueryConnectionDrawingPolicy.h
@@ -14,14 +14,14 @@ public:
 	FEnvironmentQueryConnectionDrawingPolicy(int32 InBackLayerID, int32 InFrontLayerID, float ZoomFactor, const FSlateRect& InClippingRect, FSlateWindowElementList& InDrawElements, UEdGraph* InGraphObj);
 
 	// FConnectionDrawingPolicy interface 
-	virtual void DetermineWiringStyle(UEdGraphPin* OutputPin, UEdGraphPin* InputPin, /*inout*/ float& Thickness, /*inout*/ FLinearColor& WireColor, /*inout*/ bool& bDrawBubbles, /*inout*/ bool& bBidirectional) override;
+	virtual void DetermineWiringStyle(UEdGraphPin* OutputPin, UEdGraphPin* InputPin, /*inout*/ FConnectionParams& Params) override;
 	virtual void Draw(TMap<TSharedRef<SWidget>, FArrangedWidget>& PinGeometries, FArrangedChildren& ArrangedNodes) override;
-	virtual void DrawSplineWithArrow(FGeometry& StartGeom, FGeometry& EndGeom, const FLinearColor& WireColor, float WireThickness, bool bDrawBubbles, bool Bidirectional) override;
-	virtual void DrawSplineWithArrow(const FVector2D& StartPoint, const FVector2D& EndPoint, const FLinearColor& WireColor, float WireThickness, bool bDrawBubbles, bool Bidirectional) override;
+	virtual void DrawSplineWithArrow(const FGeometry& StartGeom, const FGeometry& EndGeom, const FConnectionParams& Params) override;
+	virtual void DrawSplineWithArrow(const FVector2D& StartPoint, const FVector2D& EndPoint, const FConnectionParams& Params) override;
 	virtual void DrawPreviewConnector(const FGeometry& PinGeometry, const FVector2D& StartPoint, const FVector2D& EndPoint, UEdGraphPin* Pin) override;
-	virtual void DrawConnection(int32 LayerId, const FVector2D& Start, const FVector2D& End, const FLinearColor& InColor, float Thickness, bool bDrawBubbles) override;
+	virtual FVector2D ComputeSplineTangent(const FVector2D& Start, const FVector2D& End) const override;
 	// End of FConnectionDrawingPolicy interface
 
 protected:
-	void Internal_DrawLineWithArrow(const FVector2D& StartAnchorPoint, const FVector2D& EndAnchorPoint, const FLinearColor& WireColor, float WireThickness, bool bDrawBubbles);
+	void Internal_DrawLineWithArrow(const FVector2D& StartAnchorPoint, const FVector2D& EndAnchorPoint, const FConnectionParams& Params);
 };
diff --git a/Engine/Source/Editor/GraphEditor/Private/BlueprintConnectionDrawingPolicy.cpp b/Engine/Source/Editor/GraphEditor/Private/BlueprintConnectionDrawingPolicy.cpp
--- a/Engine/Source/Editor/GraphEditor/Private/BlueprintConnectionDrawingPolicy.cpp
+++ b/Engine/Source/Editor/GraphEditor/Private/BlueprintConnectionDrawingPolicy.cpp
@@ -340,14 +340,17 @@ FKismetConnectionDrawingPolicy::FTimePair const* FKismetConnectionDrawingPolicy:
 }
 
 // Give specific editor modes a chance to highlight this connection or darken non-interesting connections
-void FKismetConnectionDrawingPolicy::DetermineWiringStyle(UEdGraphPin* OutputPin, UEdGraphPin* InputPin, /*inout*/ float& Thickness, /*inout*/ FLinearColor& WireColor, /*inout*/bool& bDrawBubbles, /*inout*/bool& bBidirectional)
+void FKismetConnectionDrawingPolicy::DetermineWiringStyle(UEdGraphPin* OutputPin, UEdGraphPin* InputPin, /*inout*/ FConnectionParams& Params)
 {
 	// Get the schema and grab the default color from it
 	check(OutputPin);
 	check(GraphObj);
 	const UEdGraphSchema* Schema = GraphObj->GetSchema();
 
-	WireColor = Schema->GetPinTypeColor(OutputPin->PinType);
+	Params.WireColor = Schema->GetPinTypeColor(OutputPin->PinType);
+
+	UEdGraphNode* OutputNode = OutputPin->GetOwningNode();
+	UEdGraphNode* InputNode = (InputPin != nullptr) ? InputPin->GetOwningNode() : nullptr;
 
 	const bool bDeemphasizeUnhoveredPins = HoveredPins.Num() > 0;
 
@@ -359,15 +362,14 @@ void FKismetConnectionDrawingPolicy::DetermineWiringStyle(UEdGraphPin* OutputPin
 		{
 			if (CanBuildRoadmap())
 			{
-				UEdGraphNode* InputNode = InputPin->GetOwningNode();
 				// knot nodes are removed from the graph at compile time, so we 
 				// have to follow them until we find something that would have 
 				// actually executed
-				while (UK2Node_Knot* KnotNode = Cast<UK2Node_Knot>(InputNode))
+				while (UK2Node_Knot* InputKnotNode = Cast<UK2Node_Knot>(InputNode))
 				{
 					InputNode = nullptr;
 
-					UEdGraphPin* OutPin = KnotNode->GetOutputPin();
+					UEdGraphPin* OutPin = InputKnotNode->GetOutputPin();
 					if (OutPin->LinkedTo.Num() > 0)
 					{
 						check(OutPin->LinkedTo.Num() == 1);
@@ -385,36 +387,36 @@ void FKismetConnectionDrawingPolicy::DetermineWiringStyle(UEdGraphPin* OutputPin
 					if (FTimePair const* ExecTiming = BackTraceExecPath(OutputPin, ExecPaths))
 					{
 						bExecuted = true;
-						DetermineStyleOfExecWire(/*inout*/ Thickness, /*inout*/ WireColor, /*inout*/ bDrawBubbles, *ExecTiming);
+						DetermineStyleOfExecWire(/*inout*/ Params.WireThickness, /*inout*/ Params.WireColor, /*inout*/ Params.bDrawBubbles, *ExecTiming);
 					}
 				}
 				
 				if (!bExecuted)
 				{
 					// It's not followed, fade it and keep it thin
-					WireColor = ReleaseColor;
-					Thickness = ReleaseWireThickness;
+					Params.WireColor = ReleaseColor;
+					Params.WireThickness = ReleaseWireThickness;
 				}
 			}
 			else
 			{
 				// Make exec wires slightly thicker even outside of debug
-				Thickness = 3.0f;
+				Params.WireThickness = 3.0f;
 			}
 		}
 		else
 		{
 			// Array types should draw thicker
-			if( (InputPin && InputPin->PinType.bIsArray) || (OutputPin && OutputPin->PinType.bIsArray) )
+			if ((InputPin && InputPin->PinType.bIsArray) || (OutputPin && OutputPin->PinType.bIsArray))
 			{
-				Thickness = 3.0f;
+				Params.WireThickness = 3.0f;
 			}
 		}
 	}
 
 	if (bDeemphasizeUnhoveredPins)
 	{
-		ApplyHoverDeemphasis(OutputPin, InputPin, /*inout*/ Thickness, /*inout*/ WireColor);
+		ApplyHoverDeemphasis(OutputPin, InputPin, /*inout*/ Params.WireThickness, /*inout*/ Params.WireColor);
 	}
 }
 
diff --git a/Engine/Source/Editor/GraphEditor/Private/BlueprintConnectionDrawingPolicy.h b/Engine/Source/Editor/GraphEditor/Private/BlueprintConnectionDrawingPolicy.h
--- a/Engine/Source/Editor/GraphEditor/Private/BlueprintConnectionDrawingPolicy.h
+++ b/Engine/Source/Editor/GraphEditor/Private/BlueprintConnectionDrawingPolicy.h
@@ -67,7 +67,7 @@ public:
 	void CalculateEnvelopeAlphas(double ExecutionTime, /*out*/ float& AttackAlpha, /*out*/ float& SustainAlpha) const;
 
 	// FConnectionDrawingPolicy interface
-	virtual void DetermineWiringStyle(UEdGraphPin* OutputPin, UEdGraphPin* InputPin, /*inout*/ float& Thickness, /*inout*/ FLinearColor& WireColor, /*inout*/bool& bDrawBubbles, /*inout*/ bool& bBidirectional) override;
+	virtual void DetermineWiringStyle(UEdGraphPin* OutputPin, UEdGraphPin* InputPin, /*inout*/ FConnectionParams& Params) override;
 	virtual void Draw(TMap<TSharedRef<SWidget>, FArrangedWidget>& PinGeometries, FArrangedChildren& ArrangedNodes) override;
 	// End of FConnectionDrawingPolicy interface
 
diff --git a/Engine/Source/Editor/GraphEditor/Private/ConnectionDrawingPolicy.cpp b/Engine/Source/Editor/GraphEditor/Private/ConnectionDrawingPolicy.cpp
--- a/Engine/Source/Editor/GraphEditor/Private/ConnectionDrawingPolicy.cpp
+++ b/Engine/Source/Editor/GraphEditor/Private/ConnectionDrawingPolicy.cpp
@@ -96,16 +96,14 @@ FConnectionDrawingPolicy::FConnectionDrawingPolicy(int32 InBackLayerID, int32 In
 	BubbleImage = FEditorStyle::GetBrush( TEXT("Graph.ExecutionBubble") );
 }
 
-void FConnectionDrawingPolicy::DrawSplineWithArrow(const FVector2D& StartPoint, const FVector2D& EndPoint, const FLinearColor& WireColor, float WireThickness, bool bDrawBubbles, bool Bidirectional)
+void FConnectionDrawingPolicy::DrawSplineWithArrow(const FVector2D& StartPoint, const FVector2D& EndPoint, const FConnectionParams& Params)
 {
 	// Draw the spline
 	DrawConnection(
 		WireLayerID,
 		StartPoint,
 		EndPoint,
-		WireColor,
-		WireThickness,
-		bDrawBubbles);
+		Params);
 
 	// Draw the arrow
 	if (ArrowImage != nullptr)
@@ -119,12 +117,12 @@ void FConnectionDrawingPolicy::DrawSplineWithArrow(const FVector2D& StartPoint,
 			ArrowImage,
 			ClippingRect,
 			ESlateDrawEffect::None,
-			WireColor
+			Params.WireColor
 			);
 	}
 }
 
-void FConnectionDrawingPolicy::DrawSplineWithArrow(FGeometry& StartGeom, FGeometry& EndGeom, const FLinearColor& WireColor, float WireThickness, bool bDrawBubbles, bool Bidirectional)
+void FConnectionDrawingPolicy::DrawSplineWithArrow(const FGeometry& StartGeom, const FGeometry& EndGeom, const FConnectionParams& Params)
 {
 	//@TODO: These values should be pushed into the Slate style, they are compensating for a bit of
 	// empty space inside of the pin brush images.
@@ -133,7 +131,7 @@ void FConnectionDrawingPolicy::DrawSplineWithArrow(FGeometry& StartGeom, FGeomet
 	const FVector2D StartPoint = FGeometryHelper::VerticalMiddleRightOf(StartGeom) - FVector2D(StartFudgeX, 0.0f);
 	const FVector2D EndPoint = FGeometryHelper::VerticalMiddleLeftOf(EndGeom) - FVector2D(ArrowRadius.X - EndFudgeX, 0);
 
-	DrawSplineWithArrow(StartPoint, EndPoint, WireColor, WireThickness, bDrawBubbles, Bidirectional);
+	DrawSplineWithArrow(StartPoint, EndPoint, Params);
 }
 
 // Update the drawing policy with the set of hovered pins (which can be empty)
@@ -224,13 +222,14 @@ FVector2D FConnectionDrawingPolicy::ComputeSplineTangent(const FVector2D& Start,
 	return Settings->ComputeSplineTangent(Start, End);
 }
 
-void FConnectionDrawingPolicy::DrawConnection( int32 LayerId, const FVector2D& Start, const FVector2D& End, const FLinearColor& InColor, float Thickness, bool bDrawBubbles )
+void FConnectionDrawingPolicy::DrawConnection(int32 LayerId, const FVector2D& Start, const FVector2D& End, const FConnectionParams& Params)
 {
 	const FVector2D& P0 = Start;
 	const FVector2D& P1 = End;
 
-	const FVector2D P0Tangent = ComputeSplineTangent(P0, P1);
-	const FVector2D P1Tangent = P0Tangent;
+	const FVector2D SplineTangent = ComputeSplineTangent(P0, P1);
+	const FVector2D P0Tangent = SplineTangent;
+	const FVector2D P1Tangent = SplineTangent;
 
 	// Draw the spline itself
 	FSlateDrawElement::MakeDrawSpaceSpline(
@@ -239,23 +238,23 @@ void FConnectionDrawingPolicy::DrawConnection( int32 LayerId, const FVector2D& S
 		P0, P0Tangent,
 		P1, P1Tangent,
 		ClippingRect,
-		Thickness,
+		Params.WireThickness,
 		ESlateDrawEffect::None,
-		InColor
+		Params.WireColor
 	);
 
-	if (bDrawBubbles || (MidpointImage != NULL))
+	if (Params.bDrawBubbles || (MidpointImage != NULL))
 	{
 		// This table maps distance along curve to alpha
 		FInterpCurve<float> SplineReparamTable;
 		float SplineLength = MakeSplineReparamTable(P0, P0Tangent, P1, P1Tangent, SplineReparamTable);
 
 		// Draw bubbles on the spline
-		if (bDrawBubbles)
+		if (Params.bDrawBubbles)
 		{
 			const float BubbleSpacing = 64.f * ZoomFactor;
 			const float BubbleSpeed = 192.f * ZoomFactor;
-			const FVector2D BubbleSize = BubbleImage->ImageSize * ZoomFactor * 0.1f * Thickness;
+			const FVector2D BubbleSize = BubbleImage->ImageSize * ZoomFactor * 0.1f * Params.WireThickness;
 
 			float Time = (FPlatformTime::Seconds() - GStartTime);
 			const float BubbleOffset = FMath::Fmod(Time * BubbleSpeed, BubbleSpacing);
@@ -276,7 +275,7 @@ void FConnectionDrawingPolicy::DrawConnection( int32 LayerId, const FVector2D& S
 						BubbleImage,
 						ClippingRect,
 						ESlateDrawEffect::None,
-						InColor
+						Params.WireColor
 						);
 				}
 			}
@@ -308,7 +307,7 @@ void FConnectionDrawingPolicy::DrawConnection( int32 LayerId, const FVector2D& S
 				AngleInRadians,
 				TOptional<FVector2D>(),
 				FSlateDrawElement::RelativeToElement,
-				InColor
+				Params.WireColor
 				);
 		}
 	}
@@ -317,15 +316,13 @@ void FConnectionDrawingPolicy::DrawConnection( int32 LayerId, const FVector2D& S
 
 void FConnectionDrawingPolicy::DrawPreviewConnector(const FGeometry& PinGeometry, const FVector2D& StartPoint, const FVector2D& EndPoint, UEdGraphPin* Pin)
 {
-	float Thickness = 1.0f;
-	FLinearColor WireColor = FLinearColor::White;
-	bool bDrawBubbles = false;
-	bool bBiDirectional = false;
-	DetermineWiringStyle(Pin, NULL, /*inout*/ Thickness, /*inout*/ WireColor, /*inout*/ bDrawBubbles, /*inout*/ bBiDirectional);
-	DrawSplineWithArrow(StartPoint, EndPoint, WireColor, Thickness, bDrawBubbles, bBiDirectional);
+	FConnectionParams Params;
+	DetermineWiringStyle(Pin, nullptr, /*inout*/ Params);
+
+	DrawSplineWithArrow(StartPoint, EndPoint, Params);
 }
 
-void FConnectionDrawingPolicy::DetermineWiringStyle(UEdGraphPin* OutputPin, UEdGraphPin* InputPin, /*inout*/ float& Thickness, /*inout*/ FLinearColor& WireColor, /*inout*/bool& bDrawBubbles, /*inout*/ bool &bBidirectional)
+void FConnectionDrawingPolicy::DetermineWiringStyle(UEdGraphPin* OutputPin, UEdGraphPin* InputPin, /*inout*/ FConnectionParams& Params)
 {
 }
 
@@ -377,16 +374,11 @@ void FConnectionDrawingPolicy::Draw(TMap<TSharedRef<SWidget>, FArrangedWidget>&
 
 				DetermineLinkGeometry(PinGeometries, ArrangedNodes, SomePinWidget, ThePin, TargetPin, /*out*/ LinkStartWidgetGeometry, /*out*/ LinkEndWidgetGeometry);
 
-				if(( LinkEndWidgetGeometry && LinkStartWidgetGeometry ) && !IsConnectionCulled( *LinkStartWidgetGeometry, *LinkEndWidgetGeometry ))
+				if (( LinkEndWidgetGeometry && LinkStartWidgetGeometry ) && !IsConnectionCulled( *LinkStartWidgetGeometry, *LinkEndWidgetGeometry ))
 				{
-
-					float Thickness = 1.0f;
-					FLinearColor WireColor = FLinearColor::White;
-					bool bDrawBubbles = false;
-					bool bBidirectional = false;
-
-					DetermineWiringStyle(ThePin, TargetPin, /*inout*/ Thickness, /*inout*/ WireColor, /*inout*/ bDrawBubbles, /*inout*/ bBidirectional);
-					DrawSplineWithArrow(LinkStartWidgetGeometry->Geometry, LinkEndWidgetGeometry->Geometry, WireColor, Thickness, bDrawBubbles, bBidirectional);
+					FConnectionParams Params;
+					DetermineWiringStyle(ThePin, TargetPin, /*inout*/ Params);
+					DrawSplineWithArrow(LinkStartWidgetGeometry->Geometry, LinkEndWidgetGeometry->Geometry, Params);
 				}
 			}
 		}
diff --git a/Engine/Source/Editor/GraphEditor/Private/MaterialGraphConnectionDrawingPolicy.cpp b/Engine/Source/Editor/GraphEditor/Private/MaterialGraphConnectionDrawingPolicy.cpp
--- a/Engine/Source/Editor/GraphEditor/Private/MaterialGraphConnectionDrawingPolicy.cpp
+++ b/Engine/Source/Editor/GraphEditor/Private/MaterialGraphConnectionDrawingPolicy.cpp
@@ -19,29 +19,29 @@ FMaterialGraphConnectionDrawingPolicy::FMaterialGraphConnectionDrawingPolicy(int
 	HoverDeemphasisDarkFraction = 0.4f;
 }
 
-void FMaterialGraphConnectionDrawingPolicy::DetermineWiringStyle(UEdGraphPin* OutputPin, UEdGraphPin* InputPin, /*inout*/ float& Thickness, /*inout*/ FLinearColor& WireColor, /*inout*/bool& bDrawBubbles, /*inout*/ bool& bBidirectional)
+void FMaterialGraphConnectionDrawingPolicy::DetermineWiringStyle(UEdGraphPin* OutputPin, UEdGraphPin* InputPin, /*inout*/ FConnectionParams& Params)
 {
-	WireColor = MaterialGraphSchema->ActivePinColor;
+	Params.WireColor = MaterialGraphSchema->ActivePinColor;
 
 	// Have to consider both pins as the input will be an 'output' when previewing a connection
 	if (OutputPin)
 	{
 		if (!MaterialGraph->IsInputActive(OutputPin))
 		{
-			WireColor = MaterialGraphSchema->InactivePinColor;
+			Params.WireColor = MaterialGraphSchema->InactivePinColor;
 		}
 	}
 	if (InputPin)
 	{
 		if (!MaterialGraph->IsInputActive(InputPin))
 		{
-			WireColor = MaterialGraphSchema->InactivePinColor;
+			Params.WireColor = MaterialGraphSchema->InactivePinColor;
 		}
 	}
 
 	const bool bDeemphasizeUnhoveredPins = HoveredPins.Num() > 0;
 	if (bDeemphasizeUnhoveredPins)
 	{
-		ApplyHoverDeemphasis(OutputPin, InputPin, /*inout*/ Thickness, /*inout*/ WireColor);
+		ApplyHoverDeemphasis(OutputPin, InputPin, /*inout*/ Params.WireThickness, /*inout*/ Params.WireColor);
 	}
 }
diff --git a/Engine/Source/Editor/GraphEditor/Private/MaterialGraphConnectionDrawingPolicy.h b/Engine/Source/Editor/GraphEditor/Private/MaterialGraphConnectionDrawingPolicy.h
--- a/Engine/Source/Editor/GraphEditor/Private/MaterialGraphConnectionDrawingPolicy.h
+++ b/Engine/Source/Editor/GraphEditor/Private/MaterialGraphConnectionDrawingPolicy.h
@@ -18,6 +18,6 @@ public:
 	FMaterialGraphConnectionDrawingPolicy(int32 InBackLayerID, int32 InFrontLayerID, float ZoomFactor, const FSlateRect& InClippingRect, FSlateWindowElementList& InDrawElements, UEdGraph* InGraphObj);
 
 	// FConnectionDrawingPolicy interface
-	virtual void DetermineWiringStyle(UEdGraphPin* OutputPin, UEdGraphPin* InputPin, /*inout*/ float& Thickness, /*inout*/ FLinearColor& WireColor, /*inout*/bool& bDrawBubbles, /*inout*/ bool& bBidirectional) override;
+	virtual void DetermineWiringStyle(UEdGraphPin* OutputPin, UEdGraphPin* InputPin, /*inout*/ FConnectionParams& Params) override;
 	// End of FConnectionDrawingPolicy interface
 };
diff --git a/Engine/Source/Editor/GraphEditor/Private/SoundCueGraphConnectionDrawingPolicy.cpp b/Engine/Source/Editor/GraphEditor/Private/SoundCueGraphConnectionDrawingPolicy.cpp
--- a/Engine/Source/Editor/GraphEditor/Private/SoundCueGraphConnectionDrawingPolicy.cpp
+++ b/Engine/Source/Editor/GraphEditor/Private/SoundCueGraphConnectionDrawingPolicy.cpp
@@ -112,14 +112,14 @@ void FSoundCueGraphConnectionDrawingPolicy::BuildAudioFlowRoadmap()
 }
 
 // Give specific editor modes a chance to highlight this connection or darken non-interesting connections
-void FSoundCueGraphConnectionDrawingPolicy::DetermineWiringStyle(UEdGraphPin* OutputPin, UEdGraphPin* InputPin, /*inout*/ float& Thickness, /*inout*/ FLinearColor& WireColor, /*inout*/bool& bDrawBubbles, /*inout*/bool& bBidirectional)
+void FSoundCueGraphConnectionDrawingPolicy::DetermineWiringStyle(UEdGraphPin* OutputPin, UEdGraphPin* InputPin, /*inout*/ FConnectionParams& Params)
 {
 	// Get the schema and grab the default color from it
 	check(OutputPin);
 	check(GraphObj);
 	const UEdGraphSchema* Schema = GraphObj->GetSchema();
 
-	WireColor = Schema->GetPinTypeColor(OutputPin->PinType);
+	Params.WireColor = Schema->GetPinTypeColor(OutputPin->PinType);
 
 	if (InputPin == NULL)
 	{
@@ -135,17 +135,17 @@ void FSoundCueGraphConnectionDrawingPolicy::DetermineWiringStyle(UEdGraphPin* Ou
 		{
 			bExecuted = true;
 
-			Thickness = ActiveWireThickness;
-			WireColor = ActiveColor;
+			Params.WireThickness = ActiveWireThickness;
+			Params.WireColor = ActiveColor;
 
-			bDrawBubbles = true;
+			Params.bDrawBubbles = true;
 		}
 	}
 
 	if (!bExecuted)
 	{
 		// It's not followed, fade it and keep it thin
-		WireColor = InactiveColor;
-		Thickness = InactiveWireThickness;
+		Params.WireColor = InactiveColor;
+		Params.WireThickness = InactiveWireThickness;
 	}
 }
diff --git a/Engine/Source/Editor/GraphEditor/Private/SoundCueGraphConnectionDrawingPolicy.h b/Engine/Source/Editor/GraphEditor/Private/SoundCueGraphConnectionDrawingPolicy.h
--- a/Engine/Source/Editor/GraphEditor/Private/SoundCueGraphConnectionDrawingPolicy.h
+++ b/Engine/Source/Editor/GraphEditor/Private/SoundCueGraphConnectionDrawingPolicy.h
@@ -44,7 +44,7 @@ public:
 	void BuildAudioFlowRoadmap();
 
 	// FConnectionDrawingPolicy interface
-	virtual void DetermineWiringStyle(UEdGraphPin* OutputPin, UEdGraphPin* InputPin, /*inout*/ float& Thickness, /*inout*/ FLinearColor& WireColor, /*inout*/bool& bDrawBubbles, /*inout*/ bool& bBidirectional) override;
+	virtual void DetermineWiringStyle(UEdGraphPin* OutputPin, UEdGraphPin* InputPin, /*inout*/ FConnectionParams& Params) override;
 	virtual void Draw(TMap<TSharedRef<SWidget>, FArrangedWidget>& PinGeometries, FArrangedChildren& ArrangedNodes) override;
 	// End of FConnectionDrawingPolicy interface
 };
diff --git a/Engine/Source/Editor/GraphEditor/Private/StateMachineConnectionDrawingPolicy.cpp b/Engine/Source/Editor/GraphEditor/Private/StateMachineConnectionDrawingPolicy.cpp
--- a/Engine/Source/Editor/GraphEditor/Private/StateMachineConnectionDrawingPolicy.cpp
+++ b/Engine/Source/Editor/GraphEditor/Private/StateMachineConnectionDrawingPolicy.cpp
@@ -15,24 +15,24 @@ FStateMachineConnectionDrawingPolicy::FStateMachineConnectionDrawingPolicy(int32
 {
 }
 
-void FStateMachineConnectionDrawingPolicy::DetermineWiringStyle(UEdGraphPin* OutputPin, UEdGraphPin* InputPin, /*inout*/ float& Thickness, /*inout*/ FLinearColor& WireColor, /*inout*/ bool& bDrawBubbles, /*inout*/ bool& bBidirectional)
+void FStateMachineConnectionDrawingPolicy::DetermineWiringStyle(UEdGraphPin* OutputPin, UEdGraphPin* InputPin, /*inout*/ FConnectionParams& Params)
 {
-	Thickness = 1.5f;
+	Params.WireThickness = 1.5f;
 
 	if (InputPin)
 	{
 		if (UAnimStateTransitionNode* TransNode = Cast<UAnimStateTransitionNode>(InputPin->GetOwningNode()))
 		{
-			WireColor = SGraphNodeAnimTransition::StaticGetTransitionColor(TransNode, HoveredPins.Contains(InputPin));
+			Params.WireColor = SGraphNodeAnimTransition::StaticGetTransitionColor(TransNode, HoveredPins.Contains(InputPin));
 
-			bBidirectional = TransNode->Bidirectional;
+			Params.bUserFlag1 = TransNode->Bidirectional;
 		}
 	}
 
 	const bool bDeemphasizeUnhoveredPins = HoveredPins.Num() > 0;
 	if (bDeemphasizeUnhoveredPins)
 	{
-		ApplyHoverDeemphasis(OutputPin, InputPin, /*inout*/ Thickness, /*inout*/ WireColor);
+		ApplyHoverDeemphasis(OutputPin, InputPin, /*inout*/ Params.WireThickness, /*inout*/ Params.WireColor);
 	}
 }
 
@@ -100,30 +100,28 @@ void FStateMachineConnectionDrawingPolicy::Draw(TMap<TSharedRef<SWidget>, FArran
 
 void FStateMachineConnectionDrawingPolicy::DrawPreviewConnector(const FGeometry& PinGeometry, const FVector2D& StartPoint, const FVector2D& EndPoint, UEdGraphPin* Pin)
 {
-	float Thickness = 1.0f;
-	FLinearColor WireColor = FLinearColor::White;
-	bool bDrawBubbles = false;
-	bool bBiDirectional = false;
-	DetermineWiringStyle(Pin, NULL, /*inout*/ Thickness, /*inout*/ WireColor, /*inout*/ bDrawBubbles, /*inout*/ bBiDirectional);
-
+	FConnectionParams Params;
+	DetermineWiringStyle(Pin, nullptr, /*inout*/ Params);
 
 	const FVector2D SeedPoint = EndPoint;
 	const FVector2D AdjustedStartPoint = FGeometryHelper::FindClosestPointOnGeom(PinGeometry, SeedPoint);
 
-	DrawSplineWithArrow(AdjustedStartPoint, EndPoint, WireColor, Thickness, bDrawBubbles, bBiDirectional);
+	DrawSplineWithArrow(AdjustedStartPoint, EndPoint, Params);
 }
 
 
-void FStateMachineConnectionDrawingPolicy::DrawSplineWithArrow(const FVector2D& StartAnchorPoint, const FVector2D& EndAnchorPoint, const FLinearColor& WireColor, float WireThickness, bool bDrawBubbles, bool Bidirectional)
+void FStateMachineConnectionDrawingPolicy::DrawSplineWithArrow(const FVector2D& StartAnchorPoint, const FVector2D& EndAnchorPoint, const FConnectionParams& Params)
 {
-	Internal_DrawLineWithArrow(StartAnchorPoint, EndAnchorPoint, WireColor, WireThickness, bDrawBubbles);
-	if (Bidirectional)
+	Internal_DrawLineWithArrow(StartAnchorPoint, EndAnchorPoint, Params);
+
+	// Is the connection bidirectional?
+	if (Params.bUserFlag1)
 	{
-		Internal_DrawLineWithArrow(EndAnchorPoint, StartAnchorPoint, WireColor, WireThickness, bDrawBubbles);
+		Internal_DrawLineWithArrow(EndAnchorPoint, StartAnchorPoint, Params);
 	}
 }
 
-void FStateMachineConnectionDrawingPolicy::Internal_DrawLineWithArrow(const FVector2D& StartAnchorPoint, const FVector2D& EndAnchorPoint, const FLinearColor& WireColor, float WireThickness, bool bDrawBubbles)
+void FStateMachineConnectionDrawingPolicy::Internal_DrawLineWithArrow(const FVector2D& StartAnchorPoint, const FVector2D& EndAnchorPoint, const FConnectionParams& Params)
 {
 	//@TODO: Should this be scaled by zoom factor?
 	const float LineSeparationAmount = 4.5f;
@@ -139,7 +137,7 @@ void FStateMachineConnectionDrawingPolicy::Internal_DrawLineWithArrow(const FVec
 	const FVector2D EndPoint = EndAnchorPoint + DirectionBias - LengthBias;
 
 	// Draw a line/spline
-	DrawConnection(WireLayerID, StartPoint, EndPoint, WireColor, WireThickness, bDrawBubbles);
+	DrawConnection(WireLayerID, StartPoint, EndPoint, Params);
 
 	// Draw the arrow
 	const FVector2D ArrowDrawPos = EndPoint - ArrowRadius;
@@ -155,11 +153,11 @@ void FStateMachineConnectionDrawingPolicy::Internal_DrawLineWithArrow(const FVec
 		AngleInRadians,
 		TOptional<FVector2D>(),
 		FSlateDrawElement::RelativeToElement,
-		WireColor
+		Params.WireColor
 		);
 }
 
-void FStateMachineConnectionDrawingPolicy::DrawSplineWithArrow(FGeometry& StartGeom, FGeometry& EndGeom, const FLinearColor& WireColor, float WireThickness, bool bDrawBubbles, bool Bidirectional)
+void FStateMachineConnectionDrawingPolicy::DrawSplineWithArrow(const FGeometry& StartGeom, const FGeometry& EndGeom, const FConnectionParams& Params)
 {
 	// Get a reasonable seed point (halfway between the boxes)
 	const FVector2D StartCenter = FGeometryHelper::CenterOf(StartGeom);
@@ -170,31 +168,13 @@ void FStateMachineConnectionDrawingPolicy::DrawSplineWithArrow(FGeometry& StartG
 	const FVector2D StartAnchorPoint = FGeometryHelper::FindClosestPointOnGeom(StartGeom, SeedPoint);
 	const FVector2D EndAnchorPoint = FGeometryHelper::FindClosestPointOnGeom(EndGeom, SeedPoint);
 
-	DrawSplineWithArrow(StartAnchorPoint, EndAnchorPoint, WireColor, WireThickness, bDrawBubbles, Bidirectional);
+	DrawSplineWithArrow(StartAnchorPoint, EndAnchorPoint, Params);
 }
 
-void FStateMachineConnectionDrawingPolicy::DrawConnection(int32 LayerId, const FVector2D& Start, const FVector2D& End, const FLinearColor& InColor, float Thickness, bool bDrawBubbles)
+FVector2D FStateMachineConnectionDrawingPolicy::ComputeSplineTangent(const FVector2D& Start, const FVector2D& End) const
 {
-	const FVector2D& P0 = Start;
-	const FVector2D& P1 = End;
-
-	const FVector2D Delta = End-Start;
+	const FVector2D Delta = End - Start;
 	const FVector2D NormDelta = Delta.GetSafeNormal();
-	
-	const FVector2D P0Tangent = NormDelta;
-	const FVector2D P1Tangent = NormDelta;
-
-	// Draw the spline itself
-	FSlateDrawElement::MakeDrawSpaceSpline(
-		DrawElementsList,
-		LayerId,
-		P0, P0Tangent,
-		P1, P1Tangent,
-		ClippingRect,
-		Thickness,
-		ESlateDrawEffect::None,
-		InColor
-		);
 
-	//@TODO: Handle bDrawBubbles
+	return NormDelta;
 }
diff --git a/Engine/Source/Editor/GraphEditor/Private/StateMachineConnectionDrawingPolicy.h b/Engine/Source/Editor/GraphEditor/Private/StateMachineConnectionDrawingPolicy.h
--- a/Engine/Source/Editor/GraphEditor/Private/StateMachineConnectionDrawingPolicy.h
+++ b/Engine/Source/Editor/GraphEditor/Private/StateMachineConnectionDrawingPolicy.h
@@ -19,7 +19,7 @@ public:
 	FStateMachineConnectionDrawingPolicy(int32 InBackLayerID, int32 InFrontLayerID, float ZoomFactor, const FSlateRect& InClippingRect, FSlateWindowElementList& InDrawElements, UEdGraph* InGraphObj);
 
 	// FConnectionDrawingPolicy interface
-	virtual void DetermineWiringStyle(UEdGraphPin* OutputPin, UEdGraphPin* InputPin, /*inout*/ float& Thickness, /*inout*/ FLinearColor& WireColor, /*inout*/ bool& bDrawBubbles, /*inout*/ bool& bBidirectional) override;
+	virtual void DetermineWiringStyle(UEdGraphPin* OutputPin, UEdGraphPin* InputPin, /*inout*/ FConnectionParams& Params) override;
 	virtual void Draw(TMap<TSharedRef<SWidget>, FArrangedWidget>& PinGeometries, FArrangedChildren& ArrangedNodes) override;
 	virtual void DetermineLinkGeometry(
 		TMap<TSharedRef<SWidget>,
@@ -31,12 +31,12 @@ public:
 		/*out*/ FArrangedWidget*& StartWidgetGeometry,
 		/*out*/ FArrangedWidget*& EndWidgetGeometry
 		) override;
-	virtual void DrawSplineWithArrow(FGeometry& StartGeom, FGeometry& EndGeom, const FLinearColor& WireColor, float WireThickness, bool bDrawBubbles, bool Bidirectional) override;
-	virtual void DrawSplineWithArrow(const FVector2D& StartPoint, const FVector2D& EndPoint, const FLinearColor& WireColor, float WireThickness, bool bDrawBubbles, bool Bidirectional) override;
+	virtual void DrawSplineWithArrow(const FGeometry& StartGeom, const FGeometry& EndGeom, const FConnectionParams& Params) override;
+	virtual void DrawSplineWithArrow(const FVector2D& StartPoint, const FVector2D& EndPoint, const FConnectionParams& Params) override;
 	virtual void DrawPreviewConnector(const FGeometry& PinGeometry, const FVector2D& StartPoint, const FVector2D& EndPoint, UEdGraphPin* Pin) override;
-	virtual void DrawConnection(int32 LayerId, const FVector2D& Start, const FVector2D& End, const FLinearColor& InColor, float Thickness, bool bDrawBubbles) override;
+	virtual FVector2D ComputeSplineTangent(const FVector2D& Start, const FVector2D& End) const override;
 	// End of FConnectionDrawingPolicy interface
 
 protected:
-	void Internal_DrawLineWithArrow(const FVector2D& StartAnchorPoint, const FVector2D& EndAnchorPoint, const FLinearColor& WireColor, float WireThickness, bool bDrawBubbles);
+	void Internal_DrawLineWithArrow(const FVector2D& StartAnchorPoint, const FVector2D& EndAnchorPoint, const FConnectionParams& Params);
 };
diff --git a/Engine/Source/Editor/GraphEditor/Public/ConnectionDrawingPolicy.h b/Engine/Source/Editor/GraphEditor/Public/ConnectionDrawingPolicy.h
--- a/Engine/Source/Editor/GraphEditor/Public/ConnectionDrawingPolicy.h
+++ b/Engine/Source/Editor/GraphEditor/Public/ConnectionDrawingPolicy.h
@@ -23,6 +23,32 @@ public:
 	static FVector2D FindClosestPointOnGeom(const FGeometry& Geom, const FVector2D& TestPoint);
 };
 
+/////////////////////////////////////////////////////
+// FConnectionParameters
+
+struct GRAPHEDITOR_API FConnectionParams
+{
+	FLinearColor WireColor;
+	float WireThickness;
+	bool bDrawBubbles;
+	bool bUserFlag1;
+	bool bUserFlag2;
+
+	EEdGraphPinDirection StartDirection;
+	EEdGraphPinDirection EndDirection;
+
+	FConnectionParams()
+		: WireColor(FLinearColor::White)
+		, WireThickness(1.0f)
+		, bDrawBubbles(false)
+		, bUserFlag1(false)
+		, bUserFlag2(false)
+		, StartDirection(EGPD_Output)
+		, EndDirection(EGPD_Input)
+	{
+	}
+};
+
 /////////////////////////////////////////////////////
 // FConnectionDrawingPolicy
 
@@ -63,16 +89,16 @@ public:
 
 	static float MakeSplineReparamTable(const FVector2D& P0, const FVector2D& P0Tangent, const FVector2D& P1, const FVector2D& P1Tangent, FInterpCurve<float>& OutReparamTable);
 
-	virtual void DrawSplineWithArrow(const FVector2D& StartPoint, const FVector2D& EndPoint, const FLinearColor& WireColor, float WireThickness, bool bDrawBubbles, bool Bidirectional);
+	virtual void DrawSplineWithArrow(const FVector2D& StartPoint, const FVector2D& EndPoint, const FConnectionParams& Params);
 	
-	virtual void DrawSplineWithArrow(FGeometry& StartGeom, FGeometry& EndGeom, const FLinearColor& WireColor, float WireThickness, bool bDrawBubbles, bool Bidirectional);
+	virtual void DrawSplineWithArrow(const FGeometry& StartGeom, const FGeometry& EndGeom, const FConnectionParams& Params);
 
 	virtual FVector2D ComputeSplineTangent(const FVector2D& Start, const FVector2D& End) const;
-	virtual void DrawConnection(int32 LayerId, const FVector2D& Start, const FVector2D& End, const FLinearColor& InColor, float Thickness, bool bDrawBubbles);
+	virtual void DrawConnection(int32 LayerId, const FVector2D& Start, const FVector2D& End, const FConnectionParams& Params);
 	virtual void DrawPreviewConnector(const FGeometry& PinGeometry, const FVector2D& StartPoint, const FVector2D& EndPoint, UEdGraphPin* Pin);
 
 	// Give specific editor modes a chance to highlight this connection or darken non-interesting connections
-	virtual void DetermineWiringStyle(UEdGraphPin* OutputPin, UEdGraphPin* InputPin, /*inout*/ float& Thickness, /*inout*/ FLinearColor& WireColor, /*inout*/bool& bDrawBubbles, /*inout*/ bool &bBidirectional);
+	virtual void DetermineWiringStyle(UEdGraphPin* OutputPin, UEdGraphPin* InputPin, /*inout*/ FConnectionParams& Params);
 
 	virtual void Draw(TMap<TSharedRef<SWidget>, FArrangedWidget>& PinGeometries, FArrangedChildren& ArrangedNodes);
 
