fixed decorator reordering after rapid clicking in editor UE-7142

diff --git a/Engine/Source/Editor/BehaviorTreeEditor/Private/SGraphNode_BehaviorTree.cpp b/Engine/Source/Editor/BehaviorTreeEditor/Private/SGraphNode_BehaviorTree.cpp
--- a/Engine/Source/Editor/BehaviorTreeEditor/Private/SGraphNode_BehaviorTree.cpp
+++ b/Engine/Source/Editor/BehaviorTreeEditor/Private/SGraphNode_BehaviorTree.cpp
@@ -19,6 +19,35 @@
 
 #define LOCTEXT_NAMESPACE "SGraphNode_BehaviorTree"
 
+TSharedRef<FDragNodeTimed> FDragNodeTimed::New(const TSharedRef<SGraphPanel>& InGraphPanel, const TSharedRef<SGraphNode>& InDraggedNode)
+{
+	TSharedRef<FDragNodeTimed> Operation = MakeShareable(new FDragNodeTimed);
+	Operation->StartTime = FPlatformTime::Seconds();
+	Operation->GraphPanel = InGraphPanel;
+	Operation->DraggedNodes.Add(InDraggedNode);
+	// adjust the decorator away from the current mouse location a small amount based on cursor size
+	Operation->DecoratorAdjust = FSlateApplication::Get().GetCursorSize();
+	Operation->Construct();
+	return Operation;
+}
+
+TSharedRef<FDragNodeTimed> FDragNodeTimed::New(const TSharedRef<SGraphPanel>& InGraphPanel, const TArray< TSharedRef<SGraphNode> >& InDraggedNodes)
+{
+	TSharedRef<FDragNodeTimed> Operation = MakeShareable(new FDragNodeTimed);
+	Operation->StartTime = FPlatformTime::Seconds();
+	Operation->GraphPanel = InGraphPanel;
+	Operation->DraggedNodes = InDraggedNodes;
+	Operation->DecoratorAdjust = FSlateApplication::Get().GetCursorSize();
+	Operation->Construct();
+	return Operation;
+}
+
+UBehaviorTreeGraphNode* FDragNodeTimed::GetDropTargetNode() const
+{
+	return Cast<UBehaviorTreeGraphNode>(GetHoveredNode());
+}
+
+
 /////////////////////////////////////////////////////
 // SBehaviorTreePin
 
@@ -375,7 +404,6 @@ void SGraphNode_BehaviorTree::UpdateGraphNode()
 			.Padding(0)
 			.BorderBackgroundColor( this, &SGraphNode_BehaviorTree::GetBorderBackgroundColor )
 			.OnMouseButtonDown(this, &SGraphNode_BehaviorTree::OnMouseDown)
-			.OnMouseButtonUp(this, &SGraphNode_BehaviorTree::OnMouseUp)
 			[
 				SNew(SOverlay)	
 				// INPUT PIN AREA
@@ -562,11 +590,6 @@ void SGraphNode_BehaviorTree::Tick( const FGeometry& AllottedGeometry, const dou
 	SGraphNode::Tick( AllottedGeometry, InCurrentTime, InDeltaTime );
 	CachedPosition = AllottedGeometry.AbsolutePosition / AllottedGeometry.Scale;
 
-	if (bIsMouseDown)
-	{
-		MouseDownTime += InDeltaTime;
-	}
-
 	UBehaviorTreeGraphNode* MyNode = Cast<UBehaviorTreeGraphNode>(GraphNode);
 	if (MyNode && MyNode->DebuggerUpdateCounter != DebuggerStateCounter)
 	{
@@ -652,29 +675,22 @@ void SGraphNode_BehaviorTree::Tick( const FGeometry& AllottedGeometry, const dou
 
 FReply SGraphNode_BehaviorTree::OnMouseMove(const FGeometry& SenderGeometry, const FPointerEvent& MouseEvent)
 {
-	if (!MouseEvent.IsMouseButtonDown(EKeys::LeftMouseButton) && bIsMouseDown)
-	{
-		bIsMouseDown = false;
-		MouseDownTime = 0;
-	}
-
-	if (bIsMouseDown && MouseDownTime > 0.1f && FBehaviorTreeDebugger::IsPIENotSimulating())
+	if (MouseEvent.IsMouseButtonDown(EKeys::LeftMouseButton) && FBehaviorTreeDebugger::IsPIENotSimulating())
 	{
 		//if we are holding mouse over a subnode
 		if (Cast<UBehaviorTreeGraphNode>(GraphNode)->ParentNode)
 		{
 			const TSharedRef<SGraphPanel>& Panel = this->GetOwnerPanel().ToSharedRef();
 			const TSharedRef<SGraphNode>& Node = SharedThis(this);
-			return FReply::Handled().BeginDragDrop(FDragNode::New(Panel, Node));
+			return FReply::Handled().BeginDragDrop(FDragNodeTimed::New(Panel, Node));
 		}
 	}
-	return FReply::Unhandled();
-}
 
-FReply SGraphNode_BehaviorTree::OnMouseUp(const FGeometry& SenderGeometry, const FPointerEvent& MouseEvent)
-{
-	bIsMouseDown = false;
-	MouseDownTime = 0;
+	if (!MouseEvent.IsMouseButtonDown(EKeys::LeftMouseButton) && bDragMarkerVisible)
+	{
+		SetDragMarker(false);
+	}
+
 	return FReply::Unhandled();
 }
 
@@ -713,8 +729,6 @@ TSharedPtr<SGraphNode> SGraphNode_BehaviorTree::GetSubNodeUnderCursor(const FGeo
 
 FReply SGraphNode_BehaviorTree::OnMouseDown(const FGeometry& SenderGeometry, const FPointerEvent& MouseEvent)
 {
-	bIsMouseDown = true;
-
 	if (Cast<UBehaviorTreeGraphNode>(GraphNode)->ParentNode)
 	{
 		GetOwnerPanel()->SelectionManager.ClickedOnNode(GraphNode,MouseEvent);
@@ -778,20 +792,29 @@ void SGraphNode_BehaviorTree::OnDragLeave( const FDragDropEvent& DragDropEvent )
 		// Inform the Drag and Drop operation that we are not hovering any pins
 		DragConnectionOp->SetHoveredNode( TSharedPtr<SGraphNode>(NULL) );
 	}
-	SetDragMarker(false);
 
+	SetDragMarker(false);
 	SGraphNode::OnDragLeave(DragDropEvent);
 }
 
 FReply SGraphNode_BehaviorTree::OnDrop( const FGeometry& MyGeometry, const FDragDropEvent& DragDropEvent )
 {
-	TSharedPtr<FDragNode> DragNodeOp = DragDropEvent.GetOperationAs<FDragNode>();
+	SetDragMarker(false);
+
+	TSharedPtr<FDragNodeTimed> DragNodeOp = DragDropEvent.GetOperationAs<FDragNodeTimed>();
 	if (DragNodeOp.IsValid())
 	{
 		if (!DragNodeOp->IsValidOperation())
 		{
 			return FReply::Handled();
 		}
+
+		const float DragTime = float(FPlatformTime::Seconds() - DragNodeOp->StartTime);
+		if (DragTime < 0.5f)
+		{
+			return FReply::Handled();
+		}
+
 		//dropped on decorator, route to parent
 		if (Cast<UBehaviorTreeGraphNode>(GraphNode)->ParentNode != NULL)
 		{
@@ -803,11 +826,11 @@ FReply SGraphNode_BehaviorTree::OnDrop( const FGeometry& MyGeometry, const FDrag
 		const TArray< TSharedRef<SGraphNode> >& DraggedNodes = DragNodeOp->GetNodes();
 		bool bReorderOperation = true;
 
-		for (int32 i = 0; i < DraggedNodes.Num(); i++)
+		for (int32 Idx = 0; Idx < DraggedNodes.Num(); Idx++)
 		{
-			bool bIsDraggedNodeService = Cast<UBehaviorTreeGraphNode_Service>(DraggedNodes[i]->GetNodeObj()) != NULL;
+			bool bIsDraggedNodeService = Cast<UBehaviorTreeGraphNode_Service>(DraggedNodes[Idx]->GetNodeObj()) != NULL;
 
-			UBehaviorTreeGraphNode* ParentNode = (Cast<UBehaviorTreeGraphNode>(DraggedNodes[i]->GetNodeObj()))->ParentNode;
+			UBehaviorTreeGraphNode* ParentNode = (Cast<UBehaviorTreeGraphNode>(DraggedNodes[Idx]->GetNodeObj()))->ParentNode;
 			if (bReorderOperation && ParentNode != Cast<UBehaviorTreeGraphNode>(GraphNode))
 			{
 				bReorderOperation = false;
@@ -817,45 +840,56 @@ FReply SGraphNode_BehaviorTree::OnDrop( const FGeometry& MyGeometry, const FDrag
 			{
 				ParentNode->Modify();
 				TArray<UBehaviorTreeGraphNode*> & SubNodes = bIsDraggedNodeService ? ParentNode->Services : ParentNode->Decorators;
-				for (int32 j = 0; j < SubNodes.Num(); j++)
+				for (int32 SubIdx = 0; SubIdx < SubNodes.Num(); SubIdx++)
 				{
-					if (SubNodes[j] == DraggedNodes[i]->GetNodeObj())
+					if (SubNodes[SubIdx] == DraggedNodes[Idx]->GetNodeObj())
 					{
-						SubNodes.RemoveAt(j);
+						SubNodes.RemoveAt(SubIdx);
 					}
 				}
 			}
 		}
 
 		UBehaviorTreeGraphNode* BTNode = Cast<UBehaviorTreeGraphNode>(GraphNode);
-		
+		UBehaviorTreeGraphNode* DropTargetNode = DragNodeOp->GetDropTargetNode();
 		int32 InsertIndex = -1;
-		TArray< TSharedPtr<SGraphNode> > SubNodesWidgets;
-		SubNodesWidgets.Append(ServicesWidgets);
-		SubNodesWidgets.Append(DecoratorWidgets);
-		
-		for (int32 i=0; i < SubNodesWidgets.Num(); i++)
+
 		{
-			TSharedPtr<SGraphNode_BehaviorTree> BTNodeWidget = StaticCastSharedPtr<SGraphNode_BehaviorTree>(SubNodesWidgets[i]);
-			if (BTNodeWidget->bDragMarkerVisible)
+			const int32 DecoratorIdx = BTNode->Decorators.IndexOfByKey(DropTargetNode);
+			if (DecoratorIdx >= 0)
+			{
+				InsertIndex = DecoratorIdx;
+			}
+			else
 			{
-				TArray<UBehaviorTreeGraphNode*> & SubNodes = BTNodeWidget->GetNodeObj()->IsA(UBehaviorTreeGraphNode_Service::StaticClass()) ? BTNode->Services : BTNode->Decorators;
-				InsertIndex = SubNodes.IndexOfByKey(BTNodeWidget->GetNodeObj());
-				break;
+				const int32 ServiceIdx = BTNode->Services.IndexOfByKey(DropTargetNode);
+				if (ServiceIdx >= 0)
+				{
+					InsertIndex = ServiceIdx;
+				}
 			}
 		}
-		
-		for (int32 k = 0; k < DraggedNodes.Num(); k++)
+
+		for (int32 Idx = 0; Idx < DraggedNodes.Num(); Idx++)
 		{
-			UBehaviorTreeGraphNode* DraggedNode = Cast<UBehaviorTreeGraphNode>(DraggedNodes[k]->GetNodeObj());
+			UBehaviorTreeGraphNode* DraggedNode = Cast<UBehaviorTreeGraphNode>(DraggedNodes[Idx]->GetNodeObj());
 			bool bIsDraggedNodeService = Cast<UBehaviorTreeGraphNode_Service>(DraggedNode) != NULL;
 			TArray<UBehaviorTreeGraphNode*> & SubNodes = bIsDraggedNodeService ? BTNode->Services : BTNode->Decorators;
+
 			DraggedNode->Modify();
 			DraggedNode->ParentNode = BTNode;
 			BTNode->Modify();
-			InsertIndex > -1 ? SubNodes.Insert(DraggedNode, InsertIndex) : SubNodes.Add(DraggedNode);
+
+			if (InsertIndex > -1)
+			{
+				SubNodes.Insert(DraggedNode, InsertIndex);
+			}
+			else
+			{
+				SubNodes.Add(DraggedNode);
+			}
 		}
-		
+
 		if (bReorderOperation) // if updating this node is enough, do it
 		{
 			UpdateGraphNode();
diff --git a/Engine/Source/Editor/BehaviorTreeEditor/Private/SGraphNode_BehaviorTree.h b/Engine/Source/Editor/BehaviorTreeEditor/Private/SGraphNode_BehaviorTree.h
--- a/Engine/Source/Editor/BehaviorTreeEditor/Private/SGraphNode_BehaviorTree.h
+++ b/Engine/Source/Editor/BehaviorTreeEditor/Private/SGraphNode_BehaviorTree.h
@@ -2,6 +2,24 @@
 
 #pragma once
 
+class UBehaviorTreeGraphNode;
+
+class FDragNodeTimed : public FDragNode
+{
+public:
+	DRAG_DROP_OPERATOR_TYPE(FDragNodeTimed, FDragNode)
+
+	static TSharedRef<FDragNodeTimed> New(const TSharedRef<SGraphPanel>& InGraphPanel, const TSharedRef<SGraphNode>& InDraggedNode);
+	static TSharedRef<FDragNodeTimed> New(const TSharedRef<SGraphPanel>& InGraphPanel, const TArray< TSharedRef<SGraphNode> >& InDraggedNodes);
+
+	UBehaviorTreeGraphNode* GetDropTargetNode() const;
+
+	double StartTime;
+
+protected:
+	typedef FDragNode Super;
+};
+
 class SGraphNode_BehaviorTree : public SGraphNode
 {
 public:
@@ -31,9 +49,6 @@ public:
 	/** handle mouse down on the node */
 	FReply OnMouseDown(const FGeometry& SenderGeometry, const FPointerEvent& MouseEvent);
 
-	/** handle mouse up on the node */
-	FReply OnMouseUp(const FGeometry& SenderGeometry, const FPointerEvent& MouseEvent);
-
 	/** handle mouse move on the node - activate drag and drop when mouse button is being held */
 	virtual FReply OnMouseMove(const FGeometry& SenderGeometry, const FPointerEvent& MouseEvent) override;
 
@@ -67,10 +82,7 @@ public:
 	FVector2D GetCachedPosition() const { return CachedPosition; }
 
 protected:
-	float MouseDownTime;
-
 	uint32 bDragMarkerVisible : 1;
-	uint32 bIsMouseDown : 1;
 
 	uint32 bSuppressDebuggerColor : 1;
 	uint32 bSuppressDebuggerTriggers : 1;
