restoring behavior tree graph nodes when loaded asset has tree structure, but without visual representation
UE-7549

diff --git a/Engine/Source/Editor/BehaviorTreeEditor/Classes/BehaviorTreeDecoratorGraph.h b/Engine/Source/Editor/BehaviorTreeEditor/Classes/BehaviorTreeDecoratorGraph.h
--- a/Engine/Source/Editor/BehaviorTreeEditor/Classes/BehaviorTreeDecoratorGraph.h
+++ b/Engine/Source/Editor/BehaviorTreeEditor/Classes/BehaviorTreeDecoratorGraph.h
@@ -10,9 +10,15 @@ class UBehaviorTreeDecoratorGraph : public UEdGraph
 	GENERATED_UCLASS_BODY()
 
 	void CollectDecoratorData(TArray<class UBTDecorator*>& DecoratorInstances, TArray<struct FBTDecoratorLogic>& DecoratorOperations) const;
+	int32 SpawnMissingNodes(const TArray<class UBTDecorator*>& NodeInstances, const TArray<struct FBTDecoratorLogic>& Operations, int32 StartIndex);
 
 protected:
 
 	const class UBehaviorTreeDecoratorGraphNode* FindRootNode() const;
 	void CollectDecoratorDataWorker(const class UBehaviorTreeDecoratorGraphNode* Node, TArray<class UBTDecorator*>& DecoratorInstances, TArray<struct FBTDecoratorLogic>& DecoratorOperations) const;
+
+	UEdGraphPin* FindFreePin(UEdGraphNode* Node, EEdGraphPinDirection Direction);
+	UBehaviorTreeDecoratorGraphNode* SpawnMissingNodeWorker(const TArray<class UBTDecorator*>& NodeInstances, const TArray<struct FBTDecoratorLogic>& Operations, int32& Index,
+		UEdGraphNode* ParentGraphNode, int32 ChildIdx);
+
 };
diff --git a/Engine/Source/Editor/BehaviorTreeEditor/Classes/BehaviorTreeDecoratorGraphNode_Logic.h b/Engine/Source/Editor/BehaviorTreeEditor/Classes/BehaviorTreeDecoratorGraphNode_Logic.h
--- a/Engine/Source/Editor/BehaviorTreeEditor/Classes/BehaviorTreeDecoratorGraphNode_Logic.h
+++ b/Engine/Source/Editor/BehaviorTreeEditor/Classes/BehaviorTreeDecoratorGraphNode_Logic.h
@@ -33,6 +33,8 @@ class UBehaviorTreeDecoratorGraphNode_Logic : public UBehaviorTreeDecoratorGraph
 	bool CanRemovePins() const;
 	bool CanAddPins() const;
 
-	void AddInputPin();
+	UEdGraphPin* AddInputPin();
 	void RemoveInputPin(class UEdGraphPin* Pin);
+
+	EDecoratorLogicMode::Type GetLogicMode(EBTDecoratorLogic::Type Op) const;
 };
diff --git a/Engine/Source/Editor/BehaviorTreeEditor/Classes/BehaviorTreeGraph.h b/Engine/Source/Editor/BehaviorTreeEditor/Classes/BehaviorTreeGraph.h
--- a/Engine/Source/Editor/BehaviorTreeEditor/Classes/BehaviorTreeGraph.h
+++ b/Engine/Source/Editor/BehaviorTreeEditor/Classes/BehaviorTreeGraph.h
@@ -22,6 +22,7 @@ class UBehaviorTreeGraph : public UEdGraph
 	void UpdateAsset(EDebuggerFlags DebuggerFlags, bool bBumpVersion = true);
 	void UpdateAbortHighlight(struct FAbortDrawHelper& Mode0, struct FAbortDrawHelper& Mode1);
 	void CreateBTFromGraph(class UBehaviorTreeGraphNode* RootEdNode);
+	void SpawnMissingNodes();
 	void RemoveOrphanedNodes();
 	void UpdatePinConnectionTypes();
 	void UpdateDeprecatedNodes();
@@ -34,6 +35,8 @@ class UBehaviorTreeGraph : public UEdGraph
 	void LockUpdates();
 	void UnlockUpdates();
 
+	void AutoArrange();
+
 protected:
 
 	/** if set, graph modifications won't cause updates in internal tree structure (skipping UpdateAsset)
diff --git a/Engine/Source/Editor/BehaviorTreeEditor/Classes/BehaviorTreeGraphNode_CompositeDecorator.h b/Engine/Source/Editor/BehaviorTreeEditor/Classes/BehaviorTreeGraphNode_CompositeDecorator.h
--- a/Engine/Source/Editor/BehaviorTreeEditor/Classes/BehaviorTreeGraphNode_CompositeDecorator.h
+++ b/Engine/Source/Editor/BehaviorTreeEditor/Classes/BehaviorTreeGraphNode_CompositeDecorator.h
@@ -37,6 +37,7 @@ class UBehaviorTreeGraphNode_CompositeDecorator : public UBehaviorTreeGraphNode
 	virtual void PrepareForCopying() override;
 	virtual void PostCopyNode() override;
 
+	int32 SpawnMissingNodes(const TArray<class UBTDecorator*>& NodeInstances, const TArray<struct FBTDecoratorLogic>& Operations, int32 StartIndex);
 	void CollectDecoratorData(TArray<class UBTDecorator*>& NodeInstances, TArray<struct FBTDecoratorLogic>& Operations) const;
 	void SetDecoratorData(class UBTCompositeNode* InParentNode, uint8 InChildIndex);
 	void InitializeDecorator(class UBTDecorator* InnerDecorator);
diff --git a/Engine/Source/Editor/BehaviorTreeEditor/Private/BehaviorTreeDecoratorGraph.cpp b/Engine/Source/Editor/BehaviorTreeEditor/Private/BehaviorTreeDecoratorGraph.cpp
--- a/Engine/Source/Editor/BehaviorTreeEditor/Private/BehaviorTreeDecoratorGraph.cpp
+++ b/Engine/Source/Editor/BehaviorTreeEditor/Private/BehaviorTreeDecoratorGraph.cpp
@@ -81,3 +81,105 @@ void UBehaviorTreeDecoratorGraph::CollectDecoratorDataWorker(const UBehaviorTree
 		CollectDecoratorDataWorker(LinkedNodes[i], DecoratorInstances, DecoratorOperations);
 	}
 }
+
+UEdGraphPin* UBehaviorTreeDecoratorGraph::FindFreePin(UEdGraphNode* Node, EEdGraphPinDirection Direction)
+{
+	for (int32 Idx = 0; Idx < Node->Pins.Num(); Idx++)
+	{
+		if (Node->Pins[Idx] && Node->Pins[Idx]->Direction == Direction && Node->Pins[Idx]->LinkedTo.Num() == 0)
+		{
+			return Node->Pins[Idx];
+		}
+	}
+
+	return nullptr;
+}
+
+UBehaviorTreeDecoratorGraphNode* UBehaviorTreeDecoratorGraph::SpawnMissingNodeWorker(const TArray<class UBTDecorator*>& NodeInstances, const TArray<struct FBTDecoratorLogic>& Operations, int32& Index, UEdGraphNode* ParentGraphNode, int32 ChildIdx)
+{
+	int32 NumSteps = 0;
+	UBehaviorTreeDecoratorGraphNode* GraphNode = nullptr;
+	UBehaviorTreeDecoratorGraphNode_Logic* LogicNode = nullptr;
+
+	const FBTDecoratorLogic& Op = Operations[Index];
+	Index++;
+
+	if (Op.Operation == EBTDecoratorLogic::Test)
+	{
+		FGraphNodeCreator<UBehaviorTreeDecoratorGraphNode_Decorator> NodeBuilder(*this);
+		UBehaviorTreeDecoratorGraphNode_Decorator* CastedGraphNode = NodeBuilder.CreateNode();
+		NodeBuilder.Finalize();
+
+		GraphNode = CastedGraphNode;
+		CastedGraphNode->NodeInstance = NodeInstances[Op.Number];
+	}
+	else
+	{
+		FGraphNodeCreator<UBehaviorTreeDecoratorGraphNode_Logic> NodeBuilder(*this);
+		LogicNode = NodeBuilder.CreateNode();
+		LogicNode->LogicMode = LogicNode->GetLogicMode(Op.Operation);
+		NodeBuilder.Finalize();
+
+		GraphNode = LogicNode;
+		NumSteps = Op.Number;
+	}
+
+	if (GraphNode)
+	{
+		GraphNode->NodePosX = ParentGraphNode->NodePosX - 300.0f;
+		GraphNode->NodePosY = ParentGraphNode->NodePosY + ChildIdx * 100.0f;
+	}
+
+	for (int32 Idx = 0; Idx < NumSteps; Idx++)
+	{
+		UBehaviorTreeDecoratorGraphNode* ChildNode = SpawnMissingNodeWorker(NodeInstances, Operations, Index, GraphNode, Idx);
+
+		UEdGraphPin* ChildOut = FindFreePin(ChildNode, EGPD_Output);
+		UEdGraphPin* NodeIn = FindFreePin(GraphNode, EGPD_Input);
+		
+		if (NodeIn == nullptr && LogicNode)
+		{
+			NodeIn = LogicNode->AddInputPin();
+		}
+
+		if (NodeIn && ChildOut)
+		{
+			NodeIn->MakeLinkTo(ChildOut);
+		}
+	}
+
+	return GraphNode;
+}
+
+int32 UBehaviorTreeDecoratorGraph::SpawnMissingNodes(const TArray<class UBTDecorator*>& NodeInstances, const TArray<struct FBTDecoratorLogic>& Operations, int32 StartIndex)
+{
+	UBehaviorTreeDecoratorGraphNode* RootNode = nullptr;
+	for (int32 Idx = 0; Idx < Nodes.Num(); Idx++)
+	{
+		UBehaviorTreeDecoratorGraphNode_Logic* TestNode = Cast<UBehaviorTreeDecoratorGraphNode_Logic>(Nodes[Idx]);
+		if (TestNode && TestNode->LogicMode == EDecoratorLogicMode::Sink)
+		{
+			RootNode = TestNode;
+			break;
+		}
+	}
+
+	int32 NextIndex = StartIndex;
+	if (RootNode)
+	{
+		UBehaviorTreeDecoratorGraphNode* OperationRoot = SpawnMissingNodeWorker(NodeInstances, Operations, NextIndex, RootNode, 0);
+		if (OperationRoot)
+		{
+			UEdGraphPin* RootIn = FindFreePin(RootNode, EGPD_Input);
+			UEdGraphPin* OpOut = FindFreePin(OperationRoot, EGPD_Output);
+
+			RootIn->MakeLinkTo(OpOut);
+		}
+	}
+	else
+	{
+		NextIndex++;
+	}
+
+	return NextIndex;
+}
diff --git a/Engine/Source/Editor/BehaviorTreeEditor/Private/BehaviorTreeDecoratorGraphNode_Logic.cpp b/Engine/Source/Editor/BehaviorTreeEditor/Private/BehaviorTreeDecoratorGraphNode_Logic.cpp
--- a/Engine/Source/Editor/BehaviorTreeEditor/Private/BehaviorTreeDecoratorGraphNode_Logic.cpp
+++ b/Engine/Source/Editor/BehaviorTreeEditor/Private/BehaviorTreeDecoratorGraphNode_Logic.cpp
@@ -69,9 +69,9 @@ bool UBehaviorTreeDecoratorGraphNode_Logic::CanRemovePins() const
 	return (NumInputLinks > 2) && CanAddPins();
 }
 
-void UBehaviorTreeDecoratorGraphNode_Logic::AddInputPin()
+UEdGraphPin* UBehaviorTreeDecoratorGraphNode_Logic::AddInputPin()
 {
-	CreatePin(EGPD_Input, TEXT("Transition"), TEXT(""), NULL, false, false, TEXT("In"));
+	return CreatePin(EGPD_Input, TEXT("Transition"), TEXT(""), NULL, false, false, TEXT("In"));
 }
 
 void UBehaviorTreeDecoratorGraphNode_Logic::RemoveInputPin(class UEdGraphPin* Pin)
@@ -85,3 +85,9 @@ EBTDecoratorLogic::Type UBehaviorTreeDecoratorGraphNode_Logic::GetOperationType(
 	EBTDecoratorLogic::Type LogicTypes[] = { EBTDecoratorLogic::Invalid, EBTDecoratorLogic::And, EBTDecoratorLogic::Or, EBTDecoratorLogic::Not };
 	return LogicTypes[LogicMode];
 }
+
+EDecoratorLogicMode::Type UBehaviorTreeDecoratorGraphNode_Logic::GetLogicMode(EBTDecoratorLogic::Type Op) const
+{
+	EDecoratorLogicMode::Type LogicTypes[] = { EDecoratorLogicMode::Sink, EDecoratorLogicMode::Sink, EDecoratorLogicMode::And, EDecoratorLogicMode::Or, EDecoratorLogicMode::Not };
+	return LogicTypes[Op];
+}
diff --git a/Engine/Source/Editor/BehaviorTreeEditor/Private/BehaviorTreeEditor.cpp b/Engine/Source/Editor/BehaviorTreeEditor/Private/BehaviorTreeEditor.cpp
--- a/Engine/Source/Editor/BehaviorTreeEditor/Private/BehaviorTreeEditor.cpp
+++ b/Engine/Source/Editor/BehaviorTreeEditor/Private/BehaviorTreeEditor.cpp
@@ -287,6 +287,8 @@ void FBehaviorTreeEditor::RestoreBehaviorTree()
 		// Initialize the behavior tree graph
 		const UEdGraphSchema* Schema = MyGraph->GetSchema();
 		Schema->CreateDefaultNodesForGraph(*MyGraph);
+
+		MyGraph->SpawnMissingNodes();
 	}
 	else
 	{
diff --git a/Engine/Source/Editor/BehaviorTreeEditor/Private/BehaviorTreeGraph.cpp b/Engine/Source/Editor/BehaviorTreeEditor/Private/BehaviorTreeGraph.cpp
--- a/Engine/Source/Editor/BehaviorTreeEditor/Private/BehaviorTreeGraph.cpp
+++ b/Engine/Source/Editor/BehaviorTreeEditor/Private/BehaviorTreeGraph.cpp
@@ -626,6 +626,146 @@ namespace BTGraphHelpers
 		}
 	}
 
+	UEdGraphPin* FindGraphNodePin(UEdGraphNode* Node, EEdGraphPinDirection Dir)
+	{
+		UEdGraphPin* Pin = nullptr;
+		for (int32 Idx = 0; Idx < Node->Pins.Num(); Idx++)
+		{
+			if (Node->Pins[Idx]->Direction == Dir)
+			{
+				Pin = Node->Pins[Idx];
+				break;
+			}
+		}
+
+		return Pin;
+	}
+
+	void SpawnMissingDecoratorNodes(UBehaviorTreeGraphNode* GraphNode, const TArray<UBTDecorator*>& Decorators, const TArray<FBTDecoratorLogic>& DecoratorOps, UBehaviorTreeGraph* Graph)
+	{
+		if (DecoratorOps.Num() == 0)
+		{
+			for (int32 SubIdx = 0; SubIdx < Decorators.Num(); SubIdx++)
+			{
+				UBehaviorTreeGraphNode* DecoratorNode = NewObject<UBehaviorTreeGraphNode_Decorator>(Graph);
+				GraphNode->AddSubNode(DecoratorNode, Graph);
+				DecoratorNode->NodeInstance = Decorators[SubIdx];
+			}
+		}
+		else
+		{
+			int32 Idx = 0;
+			while (Idx < DecoratorOps.Num())
+			{
+				if (DecoratorOps[Idx].Operation == EBTDecoratorLogic::Test)
+				{
+					UBehaviorTreeGraphNode* DecoratorNode = NewObject<UBehaviorTreeGraphNode_Decorator>(Graph);
+					GraphNode->AddSubNode(DecoratorNode, Graph);
+					DecoratorNode->NodeInstance = Decorators[DecoratorOps[Idx].Number];
+
+					Idx++;
+				}
+				else
+				{
+					UBehaviorTreeGraphNode_CompositeDecorator* CompositeNode = NewObject<UBehaviorTreeGraphNode_CompositeDecorator>(Graph);
+					GraphNode->AddSubNode(CompositeNode, Graph);
+					
+					const int32 NextIdx = CompositeNode->SpawnMissingNodes(Decorators, DecoratorOps, Idx);
+					CompositeNode->BuildDescription();
+
+					Idx = NextIdx;
+				}
+			}
+		}
+	}
+
+	UBehaviorTreeGraphNode* SpawnMissingGraphNodesWorker(UBTNode* Node, UBehaviorTreeGraphNode* ParentGraphNode, int32 ChildIdx, int32 ParentDecoratorCount, UBehaviorTreeGraph* Graph)
+	{
+		if (Node == nullptr)
+		{
+			return nullptr;
+		}
+
+		UBehaviorTreeGraphNode* GraphNode = nullptr;
+
+		UBTCompositeNode* CompositeNode = Cast<UBTCompositeNode>(Node);
+		if (Node->IsA(UBTComposite_SimpleParallel::StaticClass()))
+		{
+			FGraphNodeCreator<UBehaviorTreeGraphNode_SimpleParallel> NodeBuilder(*Graph);
+			GraphNode = NodeBuilder.CreateNode();
+			NodeBuilder.Finalize();
+		}
+		else if (CompositeNode)
+		{
+			FGraphNodeCreator<UBehaviorTreeGraphNode_Composite> NodeBuilder(*Graph);
+			GraphNode = NodeBuilder.CreateNode();
+			NodeBuilder.Finalize();
+		}
+		else if (Node->IsA(UBTTask_RunBehavior::StaticClass()))
+		{
+			FGraphNodeCreator<UBehaviorTreeGraphNode_SubtreeTask> NodeBuilder(*Graph);
+			GraphNode = NodeBuilder.CreateNode();
+			NodeBuilder.Finalize();
+		}
+		else
+		{
+			FGraphNodeCreator<UBehaviorTreeGraphNode_Task> NodeBuilder(*Graph);
+			GraphNode = NodeBuilder.CreateNode();
+			NodeBuilder.Finalize();
+		}
+
+		if (GraphNode)
+		{
+			const int32 ParentSubNodes = ParentGraphNode->Services.Num() + ParentGraphNode->Decorators.Num();
+			GraphNode->NodePosX = ParentGraphNode->NodePosX + ChildIdx * 400.0f;
+			GraphNode->NodePosY = ParentGraphNode->NodePosY + (ParentDecoratorCount + ParentSubNodes + 1) * 75.0f;
+			GraphNode->NodeInstance = Node;
+		}
+
+		if (CompositeNode)
+		{
+			for (int32 SubIdx = 0; SubIdx < CompositeNode->Services.Num(); SubIdx++)
+			{
+				UBehaviorTreeGraphNode* ServiceNode = NewObject<UBehaviorTreeGraphNode_Service>(Graph);
+				ServiceNode->NodeInstance = CompositeNode->Services[SubIdx];
+				GraphNode->AddSubNode(ServiceNode, Graph);
+			}
+
+			UEdGraphPin* OutputPin = FindGraphNodePin(GraphNode, EGPD_Output);
+
+			for (int32 Idx = 0; Idx < CompositeNode->Children.Num(); Idx++)
+			{
+				UBTNode* ChildNode = CompositeNode->GetChildNode(Idx);
+				UBehaviorTreeGraphNode* ChildGraphNode = SpawnMissingGraphNodesWorker(ChildNode, GraphNode,
+					Idx, ParentDecoratorCount + CompositeNode->Children[Idx].Decorators.Num(), Graph);
+
+				SpawnMissingDecoratorNodes(ChildGraphNode,
+					CompositeNode->Children[Idx].Decorators,
+					CompositeNode->Children[Idx].DecoratorOps,
+					Graph);
+
+				UEdGraphPin* ChildInputPin = FindGraphNodePin(ChildGraphNode, EGPD_Input);
+
+				OutputPin->MakeLinkTo(ChildInputPin);
+			}
+		}
+
+		return GraphNode;
+	}
+
+	UBehaviorTreeGraphNode* SpawnMissingGraphNodes(UBehaviorTree* Asset, UBehaviorTreeGraphNode* ParentGraphNode, UBehaviorTreeGraph* Graph)
+	{
+		if (ParentGraphNode == nullptr || Asset == nullptr)
+		{
+			return nullptr;
+		}
+
+		UBehaviorTreeGraphNode* GraphNode = SpawnMissingGraphNodesWorker(Asset->RootNode, ParentGraphNode, 0, Asset->RootDecorators.Num(), Graph);
+		SpawnMissingDecoratorNodes(GraphNode, Asset->RootDecorators, Asset->RootDecoratorOps, Graph);
+
+		return GraphNode;
+	}
+
 } // namespace BTGraphHelpers
 
 void UBehaviorTreeGraph::CreateBTFromGraph(UBehaviorTreeGraphNode* RootEdNode)
@@ -742,6 +882,32 @@ void UBehaviorTreeGraph::RemoveOrphanedNodes()
 	}
 }
 
+void UBehaviorTreeGraph::SpawnMissingNodes()
+{
+	UBehaviorTree* BTAsset = Cast<UBehaviorTree>(GetOuter());
+	if (BTAsset)
+	{
+		UBehaviorTreeGraphNode* RootNode = nullptr;
+		for (int32 Idx = 0; Idx < Nodes.Num(); Idx++)
+		{
+			RootNode = Cast<UBehaviorTreeGraphNode_Root>(Nodes[Idx]);
+			if (RootNode)
+			{
+				break;
+			}
+		}
+
+		UBehaviorTreeGraphNode* SpawnedRootNode = BTGraphHelpers::SpawnMissingGraphNodes(BTAsset, RootNode, this);
+		if (RootNode && SpawnedRootNode)
+		{
+			UEdGraphPin* RootOutPin = BTGraphHelpers::FindGraphNodePin(RootNode, EGPD_Output);
+			UEdGraphPin* SpawnedInPin = BTGraphHelpers::FindGraphNodePin(SpawnedRootNode, EGPD_Input);
+
+			RootOutPin->MakeLinkTo(SpawnedInPin);
+		}
+	}
+}
+
 void UBehaviorTreeGraph::UpdateAbortHighlight(struct FAbortDrawHelper& Mode0, struct FAbortDrawHelper& Mode1)
 {
 	for (int32 Index = 0; Index < Nodes.Num(); ++Index)
@@ -882,3 +1048,94 @@ void UBehaviorTreeGraph::UnlockUpdates()
 
 	UpdateAsset(EDebuggerFlags::SkipDebuggerFlags);
 }
+
+namespace BTAutoArrangeHelpers
+{
+	struct FNodeBoundsInfo
+	{
+		FVector2D SubGraphBBox;
+		TArray<FNodeBoundsInfo> Children;
+	};
+
+	void AutoArrangeNodes(UBehaviorTreeGraphNode* ParentNode, FNodeBoundsInfo& BBoxTree, float PosX, float PosY)
+	{
+		int32 BBoxIndex = 0;
+
+		UEdGraphPin* Pin = BTGraphHelpers::FindGraphNodePin(ParentNode, EGPD_Output);
+		if (Pin)
+		{
+			SGraphNode::FNodeSet NodeFilter;
+			for (int32 Idx = 0; Idx < Pin->LinkedTo.Num(); Idx++)
+			{
+				UBehaviorTreeGraphNode* GraphNode = Cast<UBehaviorTreeGraphNode>(Pin->LinkedTo[Idx]->GetOwningNode());
+				if (GraphNode && BBoxTree.Children.Num() > 0)
+				{
+					AutoArrangeNodes(GraphNode, BBoxTree.Children[BBoxIndex], PosX, PosY + GraphNode->NodeWidget.Pin()->GetDesiredSize().Y * 2.5f);
+					GraphNode->NodeWidget.Pin()->MoveTo(FVector2D(BBoxTree.Children[BBoxIndex].SubGraphBBox.X / 2 - GraphNode->NodeWidget.Pin()->GetDesiredSize().X / 2 + PosX, PosY), NodeFilter);
+					PosX += BBoxTree.Children[BBoxIndex].SubGraphBBox.X + 20;
+				}
+
+				BBoxIndex++;
+			}
+		}
+	}
+
+	void GetNodeSizeInfo(UBehaviorTreeGraphNode* ParentNode, FNodeBoundsInfo& BBoxTree)
+	{
+		BBoxTree.SubGraphBBox = ParentNode->NodeWidget.Pin()->GetDesiredSize();
+		float LevelWidth = 0;
+		float LevelHeight = 0;
+
+		UEdGraphPin* Pin = BTGraphHelpers::FindGraphNodePin(ParentNode, EGPD_Output);
+		if (Pin)
+		{
+			Pin->LinkedTo.Sort(FCompareNodeXLocation());
+			for (int32 Idx = 0; Idx < Pin->LinkedTo.Num(); Idx++)
+			{
+				UBehaviorTreeGraphNode* GraphNode = Cast<UBehaviorTreeGraphNode>(Pin->LinkedTo[Idx]->GetOwningNode());
+				if (GraphNode)
+				{
+					const int32 ChildIdx = BBoxTree.Children.Add(FNodeBoundsInfo());
+					FNodeBoundsInfo& ChildBounds = BBoxTree.Children[ChildIdx];
+
+					GetNodeSizeInfo(GraphNode, ChildBounds);
+
+					LevelWidth += ChildBounds.SubGraphBBox.X + 20;
+					if (ChildBounds.SubGraphBBox.Y > LevelHeight)
+					{
+						LevelHeight = ChildBounds.SubGraphBBox.Y;
+					}
+				}
+			}
+
+			if (LevelWidth > BBoxTree.SubGraphBBox.X)
+			{
+				BBoxTree.SubGraphBBox.X = LevelWidth;
+			}
+
+			BBoxTree.SubGraphBBox.Y += LevelHeight;
+		}
+	}
+}
+
+void UBehaviorTreeGraph::AutoArrange()
+{
+	UBehaviorTreeGraphNode* RootNode = nullptr;
+	for (int32 Idx = 0; Idx < Nodes.Num(); Idx++)
+	{
+		RootNode = Cast<UBehaviorTreeGraphNode_Root>(Nodes[Idx]);
+		if (RootNode)
+		{
+			break;
+		}
+	}
+
+	BTAutoArrangeHelpers::FNodeBoundsInfo BBoxTree;
+	BTAutoArrangeHelpers::GetNodeSizeInfo(RootNode, BBoxTree);
+	BTAutoArrangeHelpers::AutoArrangeNodes(RootNode, BBoxTree, 0, RootNode->NodeWidget.Pin()->GetDesiredSize().Y * 2.5f);
+
+	RootNode->NodePosX = BBoxTree.SubGraphBBox.X / 2 - RootNode->NodeWidget.Pin()->GetDesiredSize().X / 2;
+	RootNode->NodePosY = 0;
+
+	RootNode->NodeWidget.Pin()->GetOwnerPanel()->ZoomToFit(/*bOnlySelection=*/ false);
+}
diff --git a/Engine/Source/Editor/BehaviorTreeEditor/Private/BehaviorTreeGraphNode_CompositeDecorator.cpp b/Engine/Source/Editor/BehaviorTreeEditor/Private/BehaviorTreeGraphNode_CompositeDecorator.cpp
--- a/Engine/Source/Editor/BehaviorTreeEditor/Private/BehaviorTreeGraphNode_CompositeDecorator.cpp
+++ b/Engine/Source/Editor/BehaviorTreeEditor/Private/BehaviorTreeGraphNode_CompositeDecorator.cpp
@@ -195,6 +195,19 @@ void UBehaviorTreeGraphNode_CompositeDecorator::OnInnerGraphChanged()
 	}
 }
 
+int32 UBehaviorTreeGraphNode_CompositeDecorator::SpawnMissingNodes(const TArray<class UBTDecorator*>& NodeInstances, const TArray<struct FBTDecoratorLogic>& Operations, int32 StartIndex)
+{
+	int32 FirstIdxOutside = StartIndex + 1;
+	
+	UBehaviorTreeDecoratorGraph* MyGraph = Cast<UBehaviorTreeDecoratorGraph>(BoundGraph);
+	if (MyGraph)
+	{
+		FirstIdxOutside = MyGraph->SpawnMissingNodes(NodeInstances, Operations, StartIndex);
+	}
+
+	return FirstIdxOutside;
+}
+
 void UBehaviorTreeGraphNode_CompositeDecorator::PostEditChangeProperty(struct FPropertyChangedEvent& PropertyChangedEvent)
 {
 	Super::PostEditChangeProperty(PropertyChangedEvent);
diff --git a/Engine/Source/Editor/BehaviorTreeEditor/Private/EdGraphSchema_BehaviorTree.cpp b/Engine/Source/Editor/BehaviorTreeEditor/Private/EdGraphSchema_BehaviorTree.cpp
--- a/Engine/Source/Editor/BehaviorTreeEditor/Private/EdGraphSchema_BehaviorTree.cpp
+++ b/Engine/Source/Editor/BehaviorTreeEditor/Private/EdGraphSchema_BehaviorTree.cpp
@@ -24,97 +24,14 @@ namespace
 	const int32 NodeDistance = 60;
 }
 
-struct FBBNode
-{
-	FVector2D SubGraphBBox;
-	TArray<FBBNode*> Children;
-
-	FBBNode()
-	{
-		SubGraphBBox = FVector2D(0,0);
-	}
-};
-
-void FillTreeSizeInfo(UBehaviorTreeGraphNode* ParentNode, FBBNode& BBoxTree)
-{
-	BBoxTree.SubGraphBBox = ParentNode->NodeWidget.Pin()->GetDesiredSize();
-	float LevelWidth = 0;
-	float LevelHeight = 0;
-	for (int32 i = 0; i < ParentNode->Pins.Num(); i++)
-	{
-		if (ParentNode->Pins[i]->Direction == EGPD_Output)
-		{
-			UEdGraphPin* Pin =  ParentNode->Pins[i];
-			// sort connections so that they're organized the same as user can see in the editor
-			Pin->LinkedTo.Sort(FCompareNodeXLocation());
-			for (int32 Index=0; Index < Pin->LinkedTo.Num(); ++Index)
-			{
-				UBehaviorTreeGraphNode* GraphNode = Cast<UBehaviorTreeGraphNode>(Pin->LinkedTo[Index]->GetOwningNode());
-				if (GraphNode != NULL)
-				{
-					FBBNode* ChildBBNode = new FBBNode();
-					BBoxTree.Children.Add(ChildBBNode);
-					FillTreeSizeInfo(GraphNode, *ChildBBNode);
-					LevelWidth += ChildBBNode->SubGraphBBox.X + 20;
-					if (ChildBBNode->SubGraphBBox.Y > LevelHeight)
-					{
-						LevelHeight = ChildBBNode->SubGraphBBox.Y;
-					}			
-				}
-			}
-			if (LevelWidth > BBoxTree.SubGraphBBox.X )
-			{
-				BBoxTree.SubGraphBBox.X = LevelWidth;
-			}
-			BBoxTree.SubGraphBBox.Y += LevelHeight;
-		}
-	}
-}
-
-void AutoArrangeNodes(UBehaviorTreeGraphNode* ParentNode, FBBNode& BBoxTree, float PosX, float PosY )
-{
-	int32 BBoxIndex = 0;
-	for (int32 i = 0; i < ParentNode->Pins.Num(); i++)
-	{
-		if ( ParentNode->Pins[i]->Direction == EGPD_Output)
-		{
-			UEdGraphPin* Pin =  ParentNode->Pins[i];
-			SGraphNode::FNodeSet NodeFilter;
-			for (int32 Index=0; Index < Pin->LinkedTo.Num(); ++Index)
-			{
-				UBehaviorTreeGraphNode* GraphNode = Cast<UBehaviorTreeGraphNode>(Pin->LinkedTo[Index]->GetOwningNode());
-				if (GraphNode != NULL && BBoxTree.Children.Num() > 0)
-				{
-					AutoArrangeNodes(GraphNode, *BBoxTree.Children[BBoxIndex], PosX, PosY + GraphNode->NodeWidget.Pin()->GetDesiredSize().Y * 2.5f);
-					GraphNode->NodeWidget.Pin()->MoveTo(FVector2D(BBoxTree.Children[BBoxIndex]->SubGraphBBox.X /2 - GraphNode->NodeWidget.Pin()->GetDesiredSize().X /2 + PosX, PosY), NodeFilter);
-					PosX += BBoxTree.Children[BBoxIndex]->SubGraphBBox.X + 20;
-				}
-				BBoxIndex++;
-			}
-		}
-	}
-}
-
-
 UEdGraphNode* FBehaviorTreeSchemaAction_AutoArrange::PerformAction(class UEdGraph* ParentGraph, UEdGraphPin* FromPin, const FVector2D Location, bool bSelectNewNode)
 {
-	FBBNode BBoxTree;
-	UBehaviorTreeGraphNode* RootNode = NULL;
-	for (int32 i = 0; i < ParentGraph->Nodes.Num(); i++)
+	UBehaviorTreeGraph* Graph = Cast<UBehaviorTreeGraph>(ParentGraph);
+	if (Graph)
 	{
-		if (Cast<UBehaviorTreeGraphNode_Root>(ParentGraph->Nodes[i]) != NULL)
-		{
-			RootNode = Cast<UBehaviorTreeGraphNode_Root>(ParentGraph->Nodes[i]);
-		}
+		Graph->AutoArrange();
 	}
 
-	FillTreeSizeInfo(RootNode, BBoxTree);
-	AutoArrangeNodes(RootNode, BBoxTree, 0, RootNode->NodeWidget.Pin()->GetDesiredSize().Y * 2.5f);
-	
-	RootNode->NodePosX = BBoxTree.SubGraphBBox.X / 2 - RootNode->NodeWidget.Pin()->GetDesiredSize().X / 2;
-	RootNode->NodePosY = 0;
-
-	RootNode->NodeWidget.Pin()->GetOwnerPanel()->ZoomToFit(/*bOnlySelection=*/ false);
 	return NULL;
 }
 
