compilation fix

diff --git a/Engine/Source/Editor/BehaviorTreeEditor/Private/BehaviorTreeGraph.cpp b/Engine/Source/Editor/BehaviorTreeEditor/Private/BehaviorTreeGraph.cpp
--- a/Engine/Source/Editor/BehaviorTreeEditor/Private/BehaviorTreeGraph.cpp
+++ b/Engine/Source/Editor/BehaviorTreeEditor/Private/BehaviorTreeGraph.cpp
@@ -2,6 +2,7 @@
 #include "BehaviorTreeEditorPrivatePCH.h"
 #include "BehaviorTreeEditorModule.h"
 #include "BehaviorTree/Tasks/BTTask_RunBehavior.h"
+#include "BehaviorTree/Composites/BTComposite_SimpleParallel.h"
 #include "BehaviorTree/BTDecorator.h"
 #include "BehaviorTree/BTService.h"
 #include "BehaviorTree/BehaviorTree.h"
