[INTEGRATE] Change 2413357 by [REDACTED_PERSON][REDACTED_HANDLE] on 2015/01/20 22:23:08

	More non-unity fixes.

diff --git a/Engine/Source/Editor/BehaviorTreeEditor/Private/BehaviorTreeGraph.cpp b/Engine/Source/Editor/BehaviorTreeEditor/Private/BehaviorTreeGraph.cpp
--- a/Engine/Source/Editor/BehaviorTreeEditor/Private/BehaviorTreeGraph.cpp
+++ b/Engine/Source/Editor/BehaviorTreeEditor/Private/BehaviorTreeGraph.cpp
@@ -6,6 +6,7 @@
 #include "BehaviorTree/BTDecorator.h"
 #include "BehaviorTree/BTService.h"
 #include "BehaviorTree/BehaviorTree.h"
+#include "BehaviorTree/Composites/BTComposite_SimpleParallel.h"
 
 //////////////////////////////////////////////////////////////////////////
 // BehaviorTreeGraph
diff --git a/Engine/Source/Editor/Kismet/Private/SSCSEditor.cpp b/Engine/Source/Editor/Kismet/Private/SSCSEditor.cpp
--- a/Engine/Source/Editor/Kismet/Private/SSCSEditor.cpp
+++ b/Engine/Source/Editor/Kismet/Private/SSCSEditor.cpp
@@ -36,6 +36,8 @@
 
 #include "Engine/InheritableComponentHandler.h"
 
+#include "BPVariableDragDropAction.h"
+
 #define LOCTEXT_NAMESPACE "SSCSEditor"
 
 static const FName SCS_ColumnName_ComponentClass( "ComponentClass" );
