unified shared parts of AI editors (nested nodes, copy/paste/drag operations)

diff --git a/Engine/Source/Editor/AIGraph/AIGraph.Build.cs b/Engine/Source/Editor/AIGraph/AIGraph.Build.cs
new file mode 100644
--- /dev/null
+++ b/Engine/Source/Editor/AIGraph/AIGraph.Build.cs
@@ -0,0 +1,54 @@
+// Copyright 1998-2015 Epic Games, Inc. All Rights Reserved.
+
+using UnrealBuildTool;
+
+public class AIGraph : ModuleRules
+{
+    public AIGraph(TargetInfo Target)
+    {
+        PrivateIncludePaths.AddRange(
+            new string[] {
+				"Editor/GraphEditor/Private",
+				"Editor/Kismet/Private",
+				"Editor/AIGraph/Private",
+			}
+        );
+
+        PrivateIncludePathModuleNames.AddRange(
+            new string[] {
+				"AssetRegistry",
+				"AssetTools",
+				"ContentBrowser"
+			}
+        );
+
+        PrivateDependencyModuleNames.AddRange(
+            new string[] {
+				"Core", 
+				"CoreUObject", 
+				"Engine", 
+                "RenderCore",
+                "InputCore",
+				"Slate",
+				"SlateCore",
+                "EditorStyle",
+				"UnrealEd", 
+				"MessageLog", 
+				"GraphEditor",
+                "Kismet",
+				"AnimGraph",
+				"BlueprintGraph",
+                "AIModule",
+				"ClassViewer"
+			}
+        );
+
+        DynamicallyLoadedModuleNames.AddRange(
+            new string[] { 
+				"AssetTools",
+				"AssetRegistry",
+				"ContentBrowser"
+            }
+        );
+    }
+}
diff --git a/Engine/Source/Editor/AIGraph/Classes/AIGraph.h b/Engine/Source/Editor/AIGraph/Classes/AIGraph.h
new file mode 100644
--- /dev/null
+++ b/Engine/Source/Editor/AIGraph/Classes/AIGraph.h
@@ -0,0 +1,36 @@
+// Copyright 1998-2015 Epic Games, Inc. All Rights Reserved.
+
+#pragma once
+#include "AIGraph.generated.h"
+
+UCLASS()
+class AIGRAPH_API UAIGraph : public UEdGraph
+{
+	GENERATED_UCLASS_BODY()
+
+	UPROPERTY()
+	int32 GraphVersion;
+
+	virtual void UpdateAsset(int32 UpdateFlags = 0);
+	virtual void UpdateVersion();
+	virtual void MarkVersion();
+
+	virtual void OnSubNodeDropped();
+
+	bool UpdateUnknownNodeClasses();
+	void RemoveOrphanedNodes();
+
+	bool IsLocked() const;
+	void LockUpdates();
+	void UnlockUpdates();
+
+protected:
+
+	/** if set, graph modifications won't cause updates in internal tree structure
+	 *  flag allows freezing update during heavy changes like pasting new nodes 
+	 */
+	uint32 bLockUpdates : 1;
+
+	virtual void CollectAllNodeInstances(TSet<UObject*>& NodeInstances);
+	virtual bool CanRemoveNestedObject(UObject* TestObject) const;
+};
diff --git a/Engine/Source/Editor/AIGraph/Classes/AIGraphNode.h b/Engine/Source/Editor/AIGraph/Classes/AIGraphNode.h
new file mode 100644
--- /dev/null
+++ b/Engine/Source/Editor/AIGraph/Classes/AIGraphNode.h
@@ -0,0 +1,108 @@
+// Copyright 1998-2015 Epic Games, Inc. All Rights Reserved.
+
+#pragma once
+
+#include "AIGraphTypes.h"
+#include "DiffResults.h"
+#include "AIGraphNode.generated.h"
+
+UCLASS()
+class AIGRAPH_API UAIGraphNode : public UEdGraphNode
+{
+	GENERATED_UCLASS_BODY()
+
+	/** instance class */
+	UPROPERTY()
+	struct FGraphNodeClassData ClassData;
+
+	UPROPERTY()
+	UObject* NodeInstance;
+
+	UPROPERTY(transient)
+	UAIGraphNode* ParentNode;
+
+	UPROPERTY()
+	TArray<UAIGraphNode*> SubNodes;
+
+	/** subnode index assigned during copy operation to connect nodes again on paste */
+	UPROPERTY()
+	int32 CopySubNodeIndex;
+
+	/** if set, all modifications (including delete/cut) are disabled */
+	uint32 bIsReadOnly : 1;
+
+	/** if set, this node will be always considered as subnode */
+	uint32 bIsSubNode : 1;
+
+	/** error message for node */
+	FString ErrorMessage;
+
+	// Begin UEdGraphNode Interface
+	virtual class UAIGraph* GetAIGraph();
+	virtual void AutowireNewNode(UEdGraphPin* FromPin) override;
+	virtual void PostPlacedNewNode() override;
+	virtual void PrepareForCopying() override;
+	virtual bool CanDuplicateNode() const override;
+	virtual bool CanUserDeleteNode() const override;
+	virtual void DestroyNode() override;
+	virtual FText GetTooltipText() const override;
+	virtual void NodeConnectionListChanged() override;
+	virtual bool CanCreateUnderSpecifiedSchema(const UEdGraphSchema* DesiredSchema) const override;
+	virtual void FindDiffs(class UEdGraphNode* OtherNode, struct FDiffResults& Results) override;
+	// End UEdGraphNode Interface
+
+	// Begin UObject Interface
+#if WITH_EDITOR
+	virtual void PostEditImport() override;
+	virtual void PostEditUndo() override;
+#endif
+	// End UObject
+
+	// @return the input pin for this state
+	virtual UEdGraphPin* GetInputPin(int32 InputIndex = 0) const;
+	// @return the output pin for this state
+	virtual UEdGraphPin* GetOutputPin(int32 InputIndex = 0) const;
+	virtual UEdGraph* GetBoundGraph() const { return NULL; }
+
+	virtual FText GetDescription() const;
+	virtual void PostCopyNode();
+
+	void AddSubNode(UAIGraphNode* SubNode, class UEdGraph* ParentGraph);
+	void RemoveSubNode(UAIGraphNode* SubNode);
+	virtual void RemoveAllSubNodes();
+	virtual void OnSubNodeRemoved(UAIGraphNode* SubNode);
+	virtual void OnSubNodeAdded(UAIGraphNode* SubNode);
+
+	virtual int32 FindSubNodeDropIndex(UAIGraphNode* SubNode) const;
+	virtual void InsertSubNodeAt(UAIGraphNode* SubNode, int32 DropIndex);
+
+	/** check if node is subnode */
+	virtual bool IsSubNode() const;
+
+	/** initialize instance object  */
+	virtual void InitializeInstance();
+
+	/**
+	* Finds the difference in properties of node instance
+	*
+	* @param Struct The struct of the class we are looking at
+	* @param Data The raw data for the UObject we are comparing LHS
+	* @param OtherData The raw data for the UObject we are comparing RHS
+	* @param Results The Results where differences are stored
+	* @param Diff The single result with default parameters setup
+	*/
+	void DiffProperties(UStruct* Struct, void* Data, void* OtherData, FDiffResults& Results, FDiffSingleResult& Diff);
+
+	/** reinitialize node instance */
+	virtual bool RefreshNodeClass();
+
+	/** Check if node instance uses blueprint for its implementation */
+	bool UsesBlueprint() const;
+
+	/** check if node has any errors, used for assigning colors on graph */
+	virtual bool HasErrors() const;
+
+protected:
+
+	virtual void ResetNodeOwner();
+};
diff --git a/Engine/Source/Editor/AIGraph/Classes/AIGraphSchema.h b/Engine/Source/Editor/AIGraph/Classes/AIGraphSchema.h
new file mode 100644
--- /dev/null
+++ b/Engine/Source/Editor/AIGraph/Classes/AIGraphSchema.h
@@ -0,0 +1,102 @@
+// Copyright 1998-2015 Epic Games, Inc. All Rights Reserved.
+
+#pragma once
+
+#include "AIGraphTypes.h"
+#include "AIGraphSchema.generated.h"
+
+/** Action to add a node to the graph */
+USTRUCT()
+struct AIGRAPH_API FAISchemaAction_NewNode : public FEdGraphSchemaAction
+{
+	GENERATED_USTRUCT_BODY();
+
+	/** Template of node we want to create */
+	UPROPERTY()
+	class UAIGraphNode* NodeTemplate;
+
+	FAISchemaAction_NewNode()
+		: FEdGraphSchemaAction()
+		, NodeTemplate(nullptr)
+	{}
+
+	FAISchemaAction_NewNode(const FString& InNodeCategory, const FText& InMenuDesc, const FString& InToolTip, const int32 InGrouping)
+		: FEdGraphSchemaAction(InNodeCategory, InMenuDesc, InToolTip, InGrouping)
+		, NodeTemplate(nullptr)
+	{}
+
+	// FEdGraphSchemaAction interface
+	virtual UEdGraphNode* PerformAction(class UEdGraph* ParentGraph, UEdGraphPin* FromPin, const FVector2D Location, bool bSelectNewNode = true) override;
+	virtual UEdGraphNode* PerformAction(class UEdGraph* ParentGraph, TArray<UEdGraphPin*>& FromPins, const FVector2D Location, bool bSelectNewNode = true) override;
+	virtual void AddReferencedObjects(FReferenceCollector& Collector) override;
+	// End of FEdGraphSchemaAction interface
+
+	template <typename NodeType>
+	static NodeType* SpawnNodeFromTemplate(class UEdGraph* ParentGraph, NodeType* InTemplateNode, const FVector2D Location)
+	{
+		FAISchemaAction_NewNode Action;
+		Action.NodeTemplate = InTemplateNode;
+
+		return Cast<NodeType>(Action.PerformAction(ParentGraph, nullptr, Location));
+	}
+};
+
+/** Action to add a subnode to the selected node */
+USTRUCT()
+struct AIGRAPH_API FAISchemaAction_NewSubNode : public FEdGraphSchemaAction
+{
+	GENERATED_USTRUCT_BODY();
+
+	/** Template of node we want to create */
+	UPROPERTY()
+	class UAIGraphNode* NodeTemplate;
+
+	/** parent node */
+	UPROPERTY()
+	class UAIGraphNode* ParentNode;
+
+	FAISchemaAction_NewSubNode()
+		: FEdGraphSchemaAction()
+		, NodeTemplate(nullptr)
+		, ParentNode(nullptr)
+	{}
+
+	FAISchemaAction_NewSubNode(const FString& InNodeCategory, const FText& InMenuDesc, const FString& InToolTip, const int32 InGrouping)
+		: FEdGraphSchemaAction(InNodeCategory, InMenuDesc, InToolTip, InGrouping)
+		, NodeTemplate(nullptr)
+		, ParentNode(nullptr)
+	{}
+
+	// FEdGraphSchemaAction interface
+	virtual UEdGraphNode* PerformAction(class UEdGraph* ParentGraph, UEdGraphPin* FromPin, const FVector2D Location, bool bSelectNewNode = true) override;
+	virtual UEdGraphNode* PerformAction(class UEdGraph* ParentGraph, TArray<UEdGraphPin*>& FromPins, const FVector2D Location, bool bSelectNewNode = true) override;
+	virtual void AddReferencedObjects(FReferenceCollector& Collector) override;
+	// End of FEdGraphSchemaAction interface
+};
+
+UCLASS()
+class AIGRAPH_API UAIGraphSchema : public UEdGraphSchema
+{
+	GENERATED_UCLASS_BODY()
+
+	void GetBreakLinkToSubMenuActions(class FMenuBuilder& MenuBuilder, class UEdGraphPin* InGraphPin);
+
+	// Begin EdGraphSchema interface
+	virtual void GetContextMenuActions(const UEdGraph* CurrentGraph, const UEdGraphNode* InGraphNode, const UEdGraphPin* InGraphPin, class FMenuBuilder* MenuBuilder, bool bIsDebugging) const override;
+	virtual FLinearColor GetPinTypeColor(const FEdGraphPinType& PinType) const override;
+	virtual bool ShouldHidePinDefaultValue(UEdGraphPin* Pin) const override;
+	virtual class FConnectionDrawingPolicy* CreateConnectionDrawingPolicy(int32 InBackLayerID, int32 InFrontLayerID, float InZoomFactor, const FSlateRect& InClippingRect, class FSlateWindowElementList& InDrawElements, class UEdGraph* InGraphObj) const override;
+	virtual void BreakNodeLinks(UEdGraphNode& TargetNode) const override;
+	virtual void BreakPinLinks(UEdGraphPin& TargetPin, bool bSendsNodeNotification) const override;
+	virtual void BreakSinglePinLink(UEdGraphPin* SourcePin, UEdGraphPin* TargetPin) override;
+	// End EdGraphSchema interface
+
+	virtual void GetGraphNodeContextActions(FGraphContextMenuBuilder& ContextMenuBuilder, int32 SubNodeFlags) const;
+	virtual void GetSubNodeClasses(int32 SubNodeFlags, TArray<FGraphNodeClassData>& ClassData, UClass*& GraphNodeClass) const;
+
+protected:
+
+	static TSharedPtr<FAISchemaAction_NewNode> AddNewNodeAction(FGraphActionListBuilderBase& ContextMenuBuilder, const FString& Category, const FText& MenuDesc, const FString& Tooltip);
+	static TSharedPtr<FAISchemaAction_NewSubNode> AddNewSubNodeAction(FGraphActionListBuilderBase& ContextMenuBuilder, const FString& Category, const FText& MenuDesc, const FString& Tooltip);
+};
+
diff --git a/Engine/Source/Editor/AIGraph/Classes/AIGraphTypes.h b/Engine/Source/Editor/AIGraph/Classes/AIGraphTypes.h
new file mode 100644
--- /dev/null
+++ b/Engine/Source/Editor/AIGraph/Classes/AIGraphTypes.h
@@ -0,0 +1,108 @@
+// Copyright 1998-2015 Epic Games, Inc. All Rights Reserved.
+
+#pragma once
+#include "AIGraphTypes.generated.h"
+
+USTRUCT()
+struct AIGRAPH_API FGraphNodeClassData
+{
+	GENERATED_USTRUCT_BODY()
+
+	FGraphNodeClassData() {}
+	FGraphNodeClassData(UClass* InClass, const FString& InDeprecatedMessage);
+	FGraphNodeClassData(const FString& InAssetName, const FString& InGeneratedClassPackage, const FString& InClassName, UClass* InClass);
+
+	FString ToString() const;
+	FString GetClassName() const;
+	FString GetCategory() const;
+	UClass* GetClass(bool bSilent = false);
+	bool IsAbstract() const;
+
+	FORCEINLINE bool IsBlueprint() const { return AssetName.Len() > 0; }
+	FORCEINLINE bool IsDeprecated() const { return DeprecatedMessage.Len() > 0; }
+	FORCEINLINE FString GetDeprecatedMessage() const { return DeprecatedMessage; }
+	FORCEINLINE FString GetPackageName() const { return GeneratedClassPackage; }
+
+	/** set when child class masked this one out (e.g. always use game specific class instead of engine one) */
+	uint32 bIsHidden : 1;
+
+	/** set when class wants to hide parent class from selection (just one class up hierarchy) */
+	uint32 bHideParent : 1;
+
+private:
+
+	/** pointer to uclass */
+	TWeakObjectPtr<UClass> Class;
+
+	/** path to class if it's not loaded yet */
+	UPROPERTY()
+	FString AssetName;
+	
+	UPROPERTY()
+	FString GeneratedClassPackage;
+
+	/** resolved name of class from asset data */
+	UPROPERTY()
+	FString ClassName;
+
+	/** User-defined category for this class */
+	UPROPERTY()
+	FString Category;
+
+	/** message for deprecated class */
+	FString DeprecatedMessage;
+};
+
+struct AIGRAPH_API FGraphNodeClassNode
+{
+	FGraphNodeClassData Data;
+	FString ParentClassName;
+
+	TSharedPtr<FGraphNodeClassNode> ParentNode;
+	TArray<TSharedPtr<FGraphNodeClassNode> > SubNodes;
+
+	void AddUniqueSubNode(TSharedPtr<FGraphNodeClassNode> SubNode);
+};
+
+struct AIGRAPH_API FGraphNodeClassHelper
+{
+	DECLARE_MULTICAST_DELEGATE(FOnPackageListUpdated);
+
+	FGraphNodeClassHelper(UClass* InRootClass);
+	~FGraphNodeClassHelper();
+
+	void GatherClasses(const UClass* BaseClass, TArray<FGraphNodeClassData>& AvailableClasses);
+	static FString GetDeprecationMessage(const UClass* Class);
+
+	void OnAssetAdded(const class FAssetData& AssetData);
+	void OnAssetRemoved(const class FAssetData& AssetData);
+	void InvalidateCache();
+	void OnHotReload(bool bWasTriggeredAutomatically);
+
+	static void AddUnknownClass(const FGraphNodeClassData& ClassData);
+	static bool IsClassKnown(const FGraphNodeClassData& ClassData);
+	static FOnPackageListUpdated OnPackageListUpdated;
+
+	static int32 GetObservedBlueprintClassCount(UClass* BaseNativeClass);
+	static void AddObservedBlueprintClasses(UClass* BaseNativeClass);
+	void UpdateAvailableBlueprintClasses();
+
+private:
+
+	UClass* RootNodeClass;
+	TSharedPtr<FGraphNodeClassNode> RootNode;
+	static TArray<FName> UnknownPackages;
+	static TMap<UClass*, int32> BlueprintClassCount;
+
+	TSharedPtr<FGraphNodeClassNode> CreateClassDataNode(const class FAssetData& AssetData);
+	TSharedPtr<FGraphNodeClassNode> FindBaseClassNode(TSharedPtr<FGraphNodeClassNode> Node, const FString& ClassName);
+	void FindAllSubClasses(TSharedPtr<FGraphNodeClassNode> Node, TArray<FGraphNodeClassData>& AvailableClasses);
+
+	UClass* FindAssetClass(const FString& GeneratedClassPackage, const FString& AssetName);
+	void BuildClassGraph();
+	void AddClassGraphChildren(TSharedPtr<FGraphNodeClassNode> Node, TArray<TSharedPtr<FGraphNodeClassNode> >& NodeList);
+
+	bool IsHidingClass(UClass* Class);
+	bool IsHidingParentClass(UClass* Class);
+	bool IsPackageSaved(FName PackageName);
+};
diff --git a/Engine/Source/Editor/AIGraph/Private/AIGraphConnectionDrawingPolicy.cpp b/Engine/Source/Editor/AIGraph/Private/AIGraphConnectionDrawingPolicy.cpp
new file mode 100644
--- /dev/null
+++ b/Engine/Source/Editor/AIGraph/Private/AIGraphConnectionDrawingPolicy.cpp
@@ -0,0 +1,120 @@
+// Copyright 1998-2015 Epic Games, Inc. All Rights Reserved.
+
+#include "AIGraphPrivatePCH.h"
+#include "Editor/UnrealEd/Public/Kismet2/KismetDebugUtilities.h"
+#include "Editor/UnrealEd/Public/Kismet2/BlueprintEditorUtils.h"
+#include "AIGraphConnectionDrawingPolicy.h"
+
+FAIGraphConnectionDrawingPolicy::FAIGraphConnectionDrawingPolicy(int32 InBackLayerID, int32 InFrontLayerID, float ZoomFactor, const FSlateRect& InClippingRect, FSlateWindowElementList& InDrawElements, UEdGraph* InGraphObj)
+	: FConnectionDrawingPolicy(InBackLayerID, InFrontLayerID, ZoomFactor, InClippingRect, InDrawElements)
+	, GraphObj(InGraphObj)
+{
+}
+
+void FAIGraphConnectionDrawingPolicy::DetermineWiringStyle(UEdGraphPin* OutputPin, UEdGraphPin* InputPin, /*inout*/ FConnectionParams& Params)
+{
+	Params.WireThickness = 1.5f;
+
+	const bool bDeemphasizeUnhoveredPins = HoveredPins.Num() > 0;
+	if (bDeemphasizeUnhoveredPins)
+	{
+		ApplyHoverDeemphasis(OutputPin, InputPin, /*inout*/ Params.WireThickness, /*inout*/ Params.WireColor);
+	}
+}
+
+void FAIGraphConnectionDrawingPolicy::Draw(TMap<TSharedRef<SWidget>, FArrangedWidget>& PinGeometries, FArrangedChildren& ArrangedNodes)
+{
+	// Build an acceleration structure to quickly find geometry for the nodes
+	NodeWidgetMap.Empty();
+	for (int32 NodeIndex = 0; NodeIndex < ArrangedNodes.Num(); ++NodeIndex)
+	{
+		FArrangedWidget& CurWidget = ArrangedNodes[NodeIndex];
+		TSharedRef<SGraphNode> ChildNode = StaticCastSharedRef<SGraphNode>(CurWidget.Widget);
+		NodeWidgetMap.Add(ChildNode->GetNodeObj(), NodeIndex);
+	}
+
+	// Now draw
+	FConnectionDrawingPolicy::Draw(PinGeometries, ArrangedNodes);
+}
+
+void FAIGraphConnectionDrawingPolicy::DrawPreviewConnector(const FGeometry& PinGeometry, const FVector2D& StartPoint, const FVector2D& EndPoint, UEdGraphPin* Pin)
+{
+	FConnectionParams Params;
+	DetermineWiringStyle(Pin, NULL, /*inout*/ Params);
+
+	if (Pin->Direction == EEdGraphPinDirection::EGPD_Output)
+	{
+		DrawSplineWithArrow(FGeometryHelper::FindClosestPointOnGeom(PinGeometry, EndPoint), EndPoint, Params);
+	}
+	else
+	{
+		DrawSplineWithArrow(FGeometryHelper::FindClosestPointOnGeom(PinGeometry, StartPoint), StartPoint, Params);
+	}
+}
+
+void FAIGraphConnectionDrawingPolicy::DrawSplineWithArrow(const FVector2D& StartAnchorPoint, const FVector2D& EndAnchorPoint, const FConnectionParams& Params)
+{
+	// bUserFlag1 indicates that we need to reverse the direction of connection (used by debugger)
+	const FVector2D& P0 = Params.bUserFlag1 ? EndAnchorPoint : StartAnchorPoint;
+	const FVector2D& P1 = Params.bUserFlag1 ? StartAnchorPoint : EndAnchorPoint;
+
+	Internal_DrawLineWithArrow(P0, P1, Params);
+}
+
+void FAIGraphConnectionDrawingPolicy::Internal_DrawLineWithArrow(const FVector2D& StartAnchorPoint, const FVector2D& EndAnchorPoint, const FConnectionParams& Params)
+{
+	//@TODO: Should this be scaled by zoom factor?
+	const float LineSeparationAmount = 4.5f;
+
+	const FVector2D DeltaPos = EndAnchorPoint - StartAnchorPoint;
+	const FVector2D UnitDelta = DeltaPos.GetSafeNormal();
+	const FVector2D Normal = FVector2D(DeltaPos.Y, -DeltaPos.X).GetSafeNormal();
+
+	// Come up with the final start/end points
+	const FVector2D DirectionBias = Normal * LineSeparationAmount;
+	const FVector2D LengthBias = ArrowRadius.X * UnitDelta;
+	const FVector2D StartPoint = StartAnchorPoint + DirectionBias + LengthBias;
+	const FVector2D EndPoint = EndAnchorPoint + DirectionBias - LengthBias;
+
+	// Draw a line/spline
+	DrawConnection(WireLayerID, StartPoint, EndPoint, Params);
+
+	// Draw the arrow
+	const FVector2D ArrowDrawPos = EndPoint - ArrowRadius;
+	const float AngleInRadians = FMath::Atan2(DeltaPos.Y, DeltaPos.X);
+
+	FSlateDrawElement::MakeRotatedBox(
+		DrawElementsList,
+		ArrowLayerID,
+		FPaintGeometry(ArrowDrawPos, ArrowImage->ImageSize * ZoomFactor, ZoomFactor),
+		ArrowImage,
+		ClippingRect,
+		ESlateDrawEffect::None,
+		AngleInRadians,
+		TOptional<FVector2D>(),
+		FSlateDrawElement::RelativeToElement,
+		Params.WireColor
+		);
+}
+
+void FAIGraphConnectionDrawingPolicy::DrawSplineWithArrow(const FGeometry& StartGeom, const FGeometry& EndGeom, const FConnectionParams& Params)
+{
+	// Get a reasonable seed point (halfway between the boxes)
+	const FVector2D StartCenter = FGeometryHelper::CenterOf(StartGeom);
+	const FVector2D EndCenter = FGeometryHelper::CenterOf(EndGeom);
+	const FVector2D SeedPoint = (StartCenter + EndCenter) * 0.5f;
+
+	// Find the (approximate) closest points between the two boxes
+	const FVector2D StartAnchorPoint = FGeometryHelper::FindClosestPointOnGeom(StartGeom, SeedPoint);
+	const FVector2D EndAnchorPoint = FGeometryHelper::FindClosestPointOnGeom(EndGeom, SeedPoint);
+
+	DrawSplineWithArrow(StartAnchorPoint, EndAnchorPoint, Params);
+}
+
+FVector2D FAIGraphConnectionDrawingPolicy::ComputeSplineTangent(const FVector2D& Start, const FVector2D& End) const
+{
+	const FVector2D Delta = End - Start;
+	const FVector2D NormDelta = Delta.GetSafeNormal();
+
+	return NormDelta;
+}
diff --git a/Engine/Source/Editor/AIGraph/Private/AIGraphModule.cpp b/Engine/Source/Editor/AIGraph/Private/AIGraphModule.cpp
new file mode 100644
--- /dev/null
+++ b/Engine/Source/Editor/AIGraph/Private/AIGraphModule.cpp
@@ -0,0 +1,15 @@
+// Copyright 1998-2015 Epic Games, Inc. All Rights Reserved.
+
+#include "AIGraphPrivatePCH.h"
+
+#include "ModuleManager.h"
+
+IMPLEMENT_MODULE(FAIGraphModule, AIGraph);
+
+void FAIGraphModule::StartupModule()
+{
+}
+
+void FAIGraphModule::ShutdownModule()
+{
+}
diff --git a/Engine/Source/Editor/AIGraph/Private/AIGraphPrivatePCH.h b/Engine/Source/Editor/AIGraph/Private/AIGraphPrivatePCH.h
new file mode 100644
--- /dev/null
+++ b/Engine/Source/Editor/AIGraph/Private/AIGraphPrivatePCH.h
@@ -0,0 +1,9 @@
+// Copyright 1998-2015 Epic Games, Inc. All Rights Reserved.
+
+#pragma once
+
+#include "UnrealEd.h"
+#include "GraphEditor.h"
+
+#include "AIGraphModule.h"
+#include "AIGraphClasses.h"
diff --git a/Engine/Source/Editor/AIGraph/Private/AIGraphSchema.cpp b/Engine/Source/Editor/AIGraph/Private/AIGraphSchema.cpp
new file mode 100644
--- /dev/null
+++ b/Engine/Source/Editor/AIGraph/Private/AIGraphSchema.cpp
@@ -0,0 +1,292 @@
+// Copyright 1998-2015 Epic Games, Inc. All Rights Reserved.
+
+#include "AIGraphPrivatePCH.h"
+#include "BlueprintGraphDefinitions.h"
+#include "GraphEditorActions.h"
+#include "AIGraphConnectionDrawingPolicy.h"
+#include "ScopedTransaction.h"
+#include "SGraphEditorImpl.h"
+#include "GenericCommands.h"
+
+#define LOCTEXT_NAMESPACE "AIGraph"
+#define SNAP_GRID (16) // @todo ensure this is the same as SNodePanel::GetSnapGridSize()
+
+namespace
+{
+	// Maximum distance a drag can be off a node edge to require 'push off' from node
+	const int32 NodeDistance = 60;
+}
+
+UEdGraphNode* FAISchemaAction_NewNode::PerformAction(class UEdGraph* ParentGraph, UEdGraphPin* FromPin, const FVector2D Location, bool bSelectNewNode)
+{
+	UEdGraphNode* ResultNode = NULL;
+
+	// If there is a template, we actually use it
+	if (NodeTemplate != NULL)
+	{
+		const FScopedTransaction Transaction(LOCTEXT("AddNode", "Add Node"));
+		ParentGraph->Modify();
+		if (FromPin)
+		{
+			FromPin->Modify();
+		}
+
+		NodeTemplate->SetFlags(RF_Transactional);
+
+		// set outer to be the graph so it doesn't go away
+		NodeTemplate->Rename(NULL, ParentGraph, REN_NonTransactional);
+		ParentGraph->AddNode(NodeTemplate, true);
+
+		NodeTemplate->CreateNewGuid();
+		NodeTemplate->PostPlacedNewNode();
+		NodeTemplate->AllocateDefaultPins();
+		NodeTemplate->AutowireNewNode(FromPin);
+
+		// For input pins, new node will generally overlap node being dragged off
+		// Work out if we want to visually push away from connected node
+		int32 XLocation = Location.X;
+		if (FromPin && FromPin->Direction == EGPD_Input)
+		{
+			UEdGraphNode* PinNode = FromPin->GetOwningNode();
+			const float XDelta = FMath::Abs(PinNode->NodePosX - Location.X);
+
+			if (XDelta < NodeDistance)
+			{
+				// Set location to edge of current node minus the max move distance
+				// to force node to push off from connect node enough to give selection handle
+				XLocation = PinNode->NodePosX - NodeDistance;
+			}
+		}
+
+		NodeTemplate->NodePosX = XLocation;
+		NodeTemplate->NodePosY = Location.Y;
+		NodeTemplate->SnapToGrid(SNAP_GRID);
+
+		ResultNode = NodeTemplate;
+	}
+
+	return ResultNode;
+}
+
+UEdGraphNode* FAISchemaAction_NewNode::PerformAction(class UEdGraph* ParentGraph, TArray<UEdGraphPin*>& FromPins, const FVector2D Location, bool bSelectNewNode)
+{
+	UEdGraphNode* ResultNode = NULL;
+	if (FromPins.Num() > 0)
+	{
+		ResultNode = PerformAction(ParentGraph, FromPins[0], Location);
+
+		// Try autowiring the rest of the pins
+		for (int32 Index = 1; Index < FromPins.Num(); ++Index)
+		{
+			ResultNode->AutowireNewNode(FromPins[Index]);
+		}
+	}
+	else
+	{
+		ResultNode = PerformAction(ParentGraph, NULL, Location, bSelectNewNode);
+	}
+
+	return ResultNode;
+}
+
+void FAISchemaAction_NewNode::AddReferencedObjects(FReferenceCollector& Collector)
+{
+	FEdGraphSchemaAction::AddReferencedObjects(Collector);
+
+	// These don't get saved to disk, but we want to make sure the objects don't get GC'd while the action array is around
+	Collector.AddReferencedObject(NodeTemplate);
+}
+
+UEdGraphNode* FAISchemaAction_NewSubNode::PerformAction(class UEdGraph* ParentGraph, UEdGraphPin* FromPin, const FVector2D Location, bool bSelectNewNode)
+{
+	ParentNode->AddSubNode(NodeTemplate, ParentGraph);
+	return NULL;
+}
+
+UEdGraphNode* FAISchemaAction_NewSubNode::PerformAction(class UEdGraph* ParentGraph, TArray<UEdGraphPin*>& FromPins, const FVector2D Location, bool bSelectNewNode)
+{
+	return PerformAction(ParentGraph, NULL, Location, bSelectNewNode);
+}
+
+void FAISchemaAction_NewSubNode::AddReferencedObjects(FReferenceCollector& Collector)
+{
+	FEdGraphSchemaAction::AddReferencedObjects(Collector);
+
+	// These don't get saved to disk, but we want to make sure the objects don't get GC'd while the action array is around
+	Collector.AddReferencedObject(NodeTemplate);
+}
+//////////////////////////////////////////////////////////////////////////
+
+UAIGraphSchema::UAIGraphSchema(const FObjectInitializer& ObjectInitializer) : Super(ObjectInitializer)
+{
+}
+
+TSharedPtr<FAISchemaAction_NewNode> UAIGraphSchema::AddNewNodeAction(FGraphActionListBuilderBase& ContextMenuBuilder, const FString& Category, const FText& MenuDesc, const FString& Tooltip)
+{
+	TSharedPtr<FAISchemaAction_NewNode> NewAction = TSharedPtr<FAISchemaAction_NewNode>(new FAISchemaAction_NewNode(Category, MenuDesc, Tooltip, 0));
+	ContextMenuBuilder.AddAction(NewAction);
+
+	return NewAction;
+}
+
+TSharedPtr<FAISchemaAction_NewSubNode> UAIGraphSchema::AddNewSubNodeAction(FGraphActionListBuilderBase& ContextMenuBuilder, const FString& Category, const FText& MenuDesc, const FString& Tooltip)
+{
+	TSharedPtr<FAISchemaAction_NewSubNode> NewAction = TSharedPtr<FAISchemaAction_NewSubNode>(new FAISchemaAction_NewSubNode(Category, MenuDesc, Tooltip, 0));
+	ContextMenuBuilder.AddAction(NewAction);
+	return NewAction;
+}
+
+void UAIGraphSchema::GetSubNodeClasses(int32 SubNodeFlags, TArray<FGraphNodeClassData>& ClassData, UClass*& GraphNodeClass) const
+{
+	// empty in base class
+}
+
+void UAIGraphSchema::GetGraphNodeContextActions(FGraphContextMenuBuilder& ContextMenuBuilder, int32 SubNodeFlags) const
+{
+	UEdGraph* Graph = (UEdGraph*)ContextMenuBuilder.CurrentGraph;
+	UClass* GraphNodeClass = nullptr;
+	TArray<FGraphNodeClassData> NodeClasses;
+	GetSubNodeClasses(SubNodeFlags, NodeClasses, GraphNodeClass);
+
+	if (GraphNodeClass)
+	{
+		for (const auto& NodeClass : NodeClasses)
+		{
+			const FText NodeTypeName = FText::FromString(FName::NameToDisplayString(NodeClass.ToString(), false));
+
+			UAIGraphNode* OpNode = NewObject<UAIGraphNode>(Graph, GraphNodeClass);
+			OpNode->ClassData = NodeClass;
+
+			TSharedPtr<FAISchemaAction_NewSubNode> AddOpAction = UAIGraphSchema::AddNewSubNodeAction(ContextMenuBuilder, NodeClass.GetCategory(), NodeTypeName, "");
+			AddOpAction->ParentNode = Cast<UAIGraphNode>(ContextMenuBuilder.SelectedObjects[0]);
+			AddOpAction->NodeTemplate = OpNode;
+		}
+	}
+}
+
+void UAIGraphSchema::GetContextMenuActions(const UEdGraph* CurrentGraph, const UEdGraphNode* InGraphNode, const UEdGraphPin* InGraphPin, class FMenuBuilder* MenuBuilder, bool bIsDebugging) const
+{
+	if (InGraphPin)
+	{
+		MenuBuilder->BeginSection("AIGraphSchemaPinActions", LOCTEXT("PinActionsMenuHeader", "Pin Actions"));
+		{
+			// Only display the 'Break Links' option if there is a link to break!
+			if (InGraphPin->LinkedTo.Num() > 0)
+			{
+				MenuBuilder->AddMenuEntry(FGraphEditorCommands::Get().BreakPinLinks);
+
+				// add sub menu for break link to
+				if (InGraphPin->LinkedTo.Num() > 1)
+				{
+					MenuBuilder->AddSubMenu(
+						LOCTEXT("BreakLinkTo", "Break Link To..."),
+						LOCTEXT("BreakSpecificLinks", "Break a specific link..."),
+						FNewMenuDelegate::CreateUObject((UAIGraphSchema*const)this, &UAIGraphSchema::GetBreakLinkToSubMenuActions, const_cast<UEdGraphPin*>(InGraphPin)));
+				}
+				else
+				{
+					((UAIGraphSchema*const)this)->GetBreakLinkToSubMenuActions(*MenuBuilder, const_cast<UEdGraphPin*>(InGraphPin));
+				}
+			}
+		}
+		MenuBuilder->EndSection();
+	}
+	else if (InGraphNode)
+	{
+		MenuBuilder->BeginSection("BehaviorTreeGraphSchemaNodeActions", LOCTEXT("ClassActionsMenuHeader", "Node Actions"));
+		{
+			MenuBuilder->AddMenuEntry(FGenericCommands::Get().Delete);
+			MenuBuilder->AddMenuEntry(FGenericCommands::Get().Cut);
+			MenuBuilder->AddMenuEntry(FGenericCommands::Get().Copy);
+			MenuBuilder->AddMenuEntry(FGenericCommands::Get().Duplicate);
+
+			MenuBuilder->AddMenuEntry(FGraphEditorCommands::Get().BreakNodeLinks);
+		}
+		MenuBuilder->EndSection();
+	}
+
+	Super::GetContextMenuActions(CurrentGraph, InGraphNode, InGraphPin, MenuBuilder, bIsDebugging);
+}
+
+void UAIGraphSchema::GetBreakLinkToSubMenuActions(class FMenuBuilder& MenuBuilder, UEdGraphPin* InGraphPin)
+{
+	// Make sure we have a unique name for every entry in the list
+	TMap< FString, uint32 > LinkTitleCount;
+
+	// Add all the links we could break from
+	for (TArray<class UEdGraphPin*>::TConstIterator Links(InGraphPin->LinkedTo); Links; ++Links)
+	{
+		UEdGraphPin* Pin = *Links;
+		FString TitleString = Pin->GetOwningNode()->GetNodeTitle(ENodeTitleType::ListView).ToString();
+		FText Title = FText::FromString(TitleString);
+		if (Pin->PinName != TEXT(""))
+		{
+			TitleString = FString::Printf(TEXT("%s (%s)"), *TitleString, *Pin->PinName);
+
+			// Add name of connection if possible
+			FFormatNamedArguments Args;
+			Args.Add(TEXT("NodeTitle"), Title);
+			Args.Add(TEXT("PinName"), Pin->GetDisplayName());
+			Title = FText::Format(LOCTEXT("BreakDescPin", "{NodeTitle} ({PinName})"), Args);
+		}
+
+		uint32 &Count = LinkTitleCount.FindOrAdd(TitleString);
+
+		FText Description;
+		FFormatNamedArguments Args;
+		Args.Add(TEXT("NodeTitle"), Title);
+		Args.Add(TEXT("NumberOfNodes"), Count);
+
+		if (Count == 0)
+		{
+			Description = FText::Format(LOCTEXT("BreakDesc", "Break link to {NodeTitle}"), Args);
+		}
+		else
+		{
+			Description = FText::Format(LOCTEXT("BreakDescMulti", "Break link to {NodeTitle} ({NumberOfNodes})"), Args);
+		}
+		++Count;
+
+		MenuBuilder.AddMenuEntry(Description, Description, FSlateIcon(), FUIAction(
+			FExecuteAction::CreateUObject(this, &UAIGraphSchema::BreakSinglePinLink, const_cast< UEdGraphPin* >(InGraphPin), *Links)));
+	}
+}
+
+void UAIGraphSchema::BreakNodeLinks(UEdGraphNode& TargetNode) const
+{
+	const FScopedTransaction Transaction(NSLOCTEXT("UnrealEd", "GraphEd_BreakNodeLinks", "Break Node Links"));
+
+	Super::BreakNodeLinks(TargetNode);
+}
+
+void UAIGraphSchema::BreakPinLinks(UEdGraphPin& TargetPin, bool bSendsNodeNotification) const
+{
+	const FScopedTransaction Transaction(NSLOCTEXT("UnrealEd", "GraphEd_BreakPinLinks", "Break Pin Links"));
+
+	Super::BreakPinLinks(TargetPin, bSendsNodeNotification);
+}
+
+void UAIGraphSchema::BreakSinglePinLink(UEdGraphPin* SourcePin, UEdGraphPin* TargetPin)
+{
+	const FScopedTransaction Transaction(NSLOCTEXT("UnrealEd", "GraphEd_BreakSinglePinLink", "Break Pin Link"));
+
+	Super::BreakSinglePinLink(SourcePin, TargetPin);
+}
+
+FLinearColor UAIGraphSchema::GetPinTypeColor(const FEdGraphPinType& PinType) const
+{
+	return FColor::White;
+}
+
+bool UAIGraphSchema::ShouldHidePinDefaultValue(UEdGraphPin* Pin) const
+{
+	check(Pin != NULL);
+	return Pin->bDefaultValueIsIgnored;
+}
+
+class FConnectionDrawingPolicy* UAIGraphSchema::CreateConnectionDrawingPolicy(int32 InBackLayerID, int32 InFrontLayerID, float InZoomFactor, const FSlateRect& InClippingRect, class FSlateWindowElementList& InDrawElements, class UEdGraph* InGraphObj) const
+{
+	return new FAIGraphConnectionDrawingPolicy(InBackLayerID, InFrontLayerID, InZoomFactor, InClippingRect, InDrawElements, InGraphObj);
+}
+
+#undef LOCTEXT_NAMESPACE
diff --git a/Engine/Source/Editor/AIGraph/Public/AIGraphConnectionDrawingPolicy.h b/Engine/Source/Editor/AIGraph/Public/AIGraphConnectionDrawingPolicy.h
new file mode 100644
--- /dev/null
+++ b/Engine/Source/Editor/AIGraph/Public/AIGraphConnectionDrawingPolicy.h
@@ -0,0 +1,28 @@
+// Copyright 1998-2015 Epic Games, Inc. All Rights Reserved.
+
+#pragma once
+
+#include "GraphEditorCommon.h"
+#include "ConnectionDrawingPolicy.h"
+
+class AIGRAPH_API FAIGraphConnectionDrawingPolicy : public FConnectionDrawingPolicy
+{
+protected:
+	UEdGraph* GraphObj;
+	TMap<UEdGraphNode*, int32> NodeWidgetMap;
+
+public:
+	FAIGraphConnectionDrawingPolicy(int32 InBackLayerID, int32 InFrontLayerID, float ZoomFactor, const FSlateRect& InClippingRect, FSlateWindowElementList& InDrawElements, UEdGraph* InGraphObj);
+
+	// FConnectionDrawingPolicy interface 
+	virtual void DetermineWiringStyle(UEdGraphPin* OutputPin, UEdGraphPin* InputPin, /*inout*/ FConnectionParams& Params) override;
+	virtual void Draw(TMap<TSharedRef<SWidget>, FArrangedWidget>& PinGeometries, FArrangedChildren& ArrangedNodes) override;
+	virtual void DrawSplineWithArrow(const FGeometry& StartGeom, const FGeometry& EndGeom, const FConnectionParams& Params) override;
+	virtual void DrawSplineWithArrow(const FVector2D& StartPoint, const FVector2D& EndPoint, const FConnectionParams& Params) override;
+	virtual void DrawPreviewConnector(const FGeometry& PinGeometry, const FVector2D& StartPoint, const FVector2D& EndPoint, UEdGraphPin* Pin) override;
+	virtual FVector2D ComputeSplineTangent(const FVector2D& Start, const FVector2D& End) const override;
+	// End of FConnectionDrawingPolicy interface
+
+protected:
+	void Internal_DrawLineWithArrow(const FVector2D& StartAnchorPoint, const FVector2D& EndAnchorPoint, const FConnectionParams& Params);
+};
diff --git a/Engine/Source/Editor/AIGraph/Public/AIGraphEditor.h b/Engine/Source/Editor/AIGraph/Public/AIGraphEditor.h
new file mode 100644
--- /dev/null
+++ b/Engine/Source/Editor/AIGraph/Public/AIGraphEditor.h
@@ -0,0 +1,51 @@
+// Copyright 1998-2015 Epic Games, Inc. All Rights Reserved.
+
+#pragma once
+
+#include "EditorUndoClient.h"
+
+class AIGRAPH_API FAIGraphEditor : public FEditorUndoClient
+{
+public:
+	FAIGraphEditor();
+	virtual ~FAIGraphEditor();
+
+	FGraphPanelSelectionSet GetSelectedNodes() const;
+	virtual void OnSelectedNodesChanged(const TSet<class UObject*>& NewSelection);
+
+	// Begin FEditorUndoClient Interface
+	virtual void PostUndo(bool bSuccess) override;
+	virtual void PostRedo(bool bSuccess) override;
+	// End of FEditorUndoClient
+
+	void CreateCommandList();
+
+	// Delegates for graph editor commands
+	void SelectAllNodes();
+	bool CanSelectAllNodes() const;
+	void DeleteSelectedNodes();
+	bool CanDeleteNodes() const;
+	void DeleteSelectedDuplicatableNodes();
+	void CutSelectedNodes();
+	bool CanCutNodes() const;
+	void CopySelectedNodes();
+	bool CanCopyNodes() const;
+	void PasteNodes();
+	void PasteNodesHere(const FVector2D& Location);
+	bool CanPasteNodes() const;
+	void DuplicateNodes();
+	bool CanDuplicateNodes() const;
+
+	virtual void OnClassListUpdated();
+
+protected:
+
+	/** Currently focused graph */
+	TWeakPtr<SGraphEditor> UpdateGraphEdPtr;
+
+	/** The command list for this editor */
+	TSharedPtr<FUICommandList> GraphEditorCommands;
+
+	/** Handle to the registered OnClassListUpdated delegate */
+	FDelegateHandle OnClassListUpdatedDelegateHandle;
+};
diff --git a/Engine/Source/Editor/AIGraph/Public/AIGraphModule.h b/Engine/Source/Editor/AIGraph/Public/AIGraphModule.h
new file mode 100644
--- /dev/null
+++ b/Engine/Source/Editor/AIGraph/Public/AIGraphModule.h
@@ -0,0 +1,13 @@
+// Copyright 1998-2015 Epic Games, Inc. All Rights Reserved.
+
+#pragma once
+
+#include "ModuleInterface.h"
+
+class FAIGraphModule : public IModuleInterface
+{
+public:
+	// IModuleInterface
+	virtual void StartupModule() override;
+	virtual void ShutdownModule() override;
+};
diff --git a/Engine/Source/Editor/AIGraph/Public/SGraphEditorActionMenuAI.h b/Engine/Source/Editor/AIGraph/Public/SGraphEditorActionMenuAI.h
new file mode 100644
--- /dev/null
+++ b/Engine/Source/Editor/AIGraph/Public/SGraphEditorActionMenuAI.h
@@ -0,0 +1,49 @@
+// Copyright 1998-2015 Epic Games, Inc. All Rights Reserved.
+
+#pragma once
+#include "SGraphActionMenu.h"
+
+/////////////////////////////////////////////////////////////////////////////////////////////////
+
+class AIGRAPH_API SGraphEditorActionMenuAI : public SBorder
+{
+public:
+	SLATE_BEGIN_ARGS(SGraphEditorActionMenuAI)
+		: _GraphObj( static_cast<UEdGraph*>(NULL) )
+		,_GraphNode(NULL)
+		, _NewNodePosition( FVector2D::ZeroVector )
+		, _AutoExpandActionMenu( false )
+		, _SubNodeFlags(0)
+		{}
+
+		SLATE_ARGUMENT( UEdGraph*, GraphObj )
+		SLATE_ARGUMENT( UAIGraphNode*, GraphNode)
+		SLATE_ARGUMENT( FVector2D, NewNodePosition )
+		SLATE_ARGUMENT( TArray<UEdGraphPin*>, DraggedFromPins )
+		SLATE_ARGUMENT( SGraphEditor::FActionMenuClosed, OnClosedCallback )
+		SLATE_ARGUMENT( bool, AutoExpandActionMenu )
+		SLATE_ARGUMENT( int32, SubNodeFlags)
+	SLATE_END_ARGS()
+
+	void Construct( const FArguments& InArgs );
+
+	~SGraphEditorActionMenuAI();
+
+	TSharedRef<SEditableTextBox> GetFilterTextBox();
+
+protected:
+	UEdGraph* GraphObj;
+	UAIGraphNode* GraphNode;
+	TArray<UEdGraphPin*> DraggedFromPins;
+	FVector2D NewNodePosition;
+	bool AutoExpandActionMenu;
+	int32 SubNodeFlags;
+
+	SGraphEditor::FActionMenuClosed OnClosedCallback;
+	TSharedPtr<SGraphActionMenu> GraphActionMenu;
+
+	void OnActionSelected( const TArray< TSharedPtr<FEdGraphSchemaAction> >& SelectedAction );
+
+	/** Callback used to populate all actions list in SGraphActionMenu */
+	void CollectAllActions(FGraphActionListBuilderBase& OutAllActions);
+};
diff --git a/Engine/Source/Editor/AIGraph/Public/SGraphNodeAI.h b/Engine/Source/Editor/AIGraph/Public/SGraphNodeAI.h
new file mode 100644
--- /dev/null
+++ b/Engine/Source/Editor/AIGraph/Public/SGraphNodeAI.h
@@ -0,0 +1,83 @@
+// Copyright 1998-2015 Epic Games, Inc. All Rights Reserved.
+
+#pragma once
+
+#include "GraphEditorCommon.h"
+
+class UAIGraphNode;
+
+class FDragAIGraphNode : public FDragNode
+{
+public:
+	DRAG_DROP_OPERATOR_TYPE(FDragAIGraphNode, FDragNode)
+
+	static TSharedRef<FDragAIGraphNode> New(const TSharedRef<SGraphPanel>& InGraphPanel, const TSharedRef<SGraphNode>& InDraggedNode);
+	static TSharedRef<FDragAIGraphNode> New(const TSharedRef<SGraphPanel>& InGraphPanel, const TArray< TSharedRef<SGraphNode> >& InDraggedNodes);
+
+	UAIGraphNode* GetDropTargetNode() const;
+
+	double StartTime;
+
+protected:
+	typedef FDragNode Super;
+};
+
+class AIGRAPH_API SGraphNodeAI : public SGraphNode
+{
+public:
+	SLATE_BEGIN_ARGS(SGraphNodeAI){}
+	SLATE_END_ARGS()
+
+	void Construct(const FArguments& InArgs, UAIGraphNode* InNode);
+
+	// SGraphNode interface
+	virtual TSharedPtr<SToolTip> GetComplexTooltip() override;
+	virtual void OnDragEnter(const FGeometry& MyGeometry, const FDragDropEvent& DragDropEvent) override;
+	virtual FReply OnDragOver(const FGeometry& MyGeometry, const FDragDropEvent& DragDropEvent) override;
+	virtual FReply OnDrop(const FGeometry& MyGeometry, const FDragDropEvent& DragDropEvent) override;
+	virtual void OnDragLeave(const FDragDropEvent& DragDropEvent) override;
+	virtual FReply OnMouseMove(const FGeometry& SenderGeometry, const FPointerEvent& MouseEvent) override;
+	virtual void SetOwner(const TSharedRef<SGraphPanel>& OwnerPanel) override;
+	virtual void AddPin(const TSharedRef<SGraphPin>& PinToAdd) override;
+	// End of SGraphNode interface
+
+	/** handle mouse down on the node */
+	FReply OnMouseDown(const FGeometry& SenderGeometry, const FPointerEvent& MouseEvent);
+
+	/** adds subnode widget inside current node */
+	virtual void AddSubNode(TSharedPtr<SGraphNode> SubNodeWidget);
+
+	/** gets decorator or service node if one is found under mouse cursor */
+	TSharedPtr<SGraphNode> GetSubNodeUnderCursor(const FGeometry& WidgetGeometry, const FPointerEvent& MouseEvent);
+
+	/** gets drag over marker visibility */
+	EVisibility GetDragOverMarkerVisibility() const;
+
+	/** sets drag marker visible or collapsed on this node */
+	void SetDragMarker(bool bEnabled);
+
+protected:
+	TArray< TSharedPtr<SGraphNode> > SubNodes;
+
+	uint32 bDragMarkerVisible : 1;
+
+	virtual FText GetDescription() const;
+	virtual FText GetPreviewCornerText() const;
+	virtual const FSlateBrush* GetNameIcon() const;
+};
+
+class AIGRAPH_API SGraphPinAI : public SGraphPin
+{
+public:
+	SLATE_BEGIN_ARGS(SGraphPinAI){}
+	SLATE_END_ARGS()
+
+	void Construct(const FArguments& InArgs, UEdGraphPin* InPin);
+protected:
+	// Begin SGraphPin interface
+	virtual FSlateColor GetPinColor() const override;
+	virtual TSharedRef<SWidget>	GetDefaultValueWidget() override;
+	// End SGraphPin interface
+
+	const FSlateBrush* GetPinBorder() const;
+};
diff --git a/Engine/Source/Editor/BehaviorTreeEditor/BehaviorTreeEditor.Build.cs b/Engine/Source/Editor/BehaviorTreeEditor/BehaviorTreeEditor.Build.cs
--- a/Engine/Source/Editor/BehaviorTreeEditor/BehaviorTreeEditor.Build.cs
+++ b/Engine/Source/Editor/BehaviorTreeEditor/BehaviorTreeEditor.Build.cs
@@ -9,7 +9,7 @@ public class BehaviorTreeEditor : ModuleRules
         PrivateIncludePaths.AddRange(
             new string[] {
 				"Editor/GraphEditor/Private",
-				"Editor/Kismet/Private",
+				"Editor/AIGraph/Private",
 				"Editor/BehaviorTreeEditor/Private",
 			}
 		);
@@ -40,6 +40,7 @@ public class BehaviorTreeEditor : ModuleRules
                 "PropertyEditor",
 				"AnimGraph",
 				"BlueprintGraph",
+                "AIGraph",
                 "AIModule",
                 "GameplayDebugger",
 				"ClassViewer"
diff --git a/Engine/Source/Editor/BehaviorTreeEditor/Classes/BehaviorTreeDecoratorGraphNode_Decorator.h b/Engine/Source/Editor/BehaviorTreeEditor/Classes/BehaviorTreeDecoratorGraphNode_Decorator.h
--- a/Engine/Source/Editor/BehaviorTreeEditor/Classes/BehaviorTreeDecoratorGraphNode_Decorator.h
+++ b/Engine/Source/Editor/BehaviorTreeEditor/Classes/BehaviorTreeDecoratorGraphNode_Decorator.h
@@ -1,6 +1,7 @@
 // Copyright 1998-2015 Epic Games, Inc. All Rights Reserved.
 
 #pragma once
+
 #include "BehaviorTreeEditorTypes.h"
 #include "BehaviorTreeDecoratorGraphNode_Decorator.generated.h"
 
@@ -12,7 +13,8 @@ class UBehaviorTreeDecoratorGraphNode_Decorator : public UBehaviorTreeDecoratorG
 	UPROPERTY()
 	UObject* NodeInstance;
 
-	struct FClassData ClassData;
+	UPROPERTY()
+	struct FGraphNodeClassData ClassData;
 
 	virtual void PostPlacedNewNode() override;
 	virtual void AllocateDefaultPins() override;
diff --git a/Engine/Source/Editor/BehaviorTreeEditor/Classes/BehaviorTreeEditorTypes.h b/Engine/Source/Editor/BehaviorTreeEditor/Classes/BehaviorTreeEditorTypes.h
--- a/Engine/Source/Editor/BehaviorTreeEditor/Classes/BehaviorTreeEditorTypes.h
+++ b/Engine/Source/Editor/BehaviorTreeEditor/Classes/BehaviorTreeEditorTypes.h
@@ -1,6 +1,8 @@
 // Copyright 1998-2015 Epic Games, Inc. All Rights Reserved.
 
 #pragma once
+
+#include "AIGraphTypes.h"
 #include "BehaviorTreeEditorTypes.generated.h"
 
 struct FAbortDrawHelper
@@ -13,124 +15,6 @@ struct FAbortDrawHelper
 	FAbortDrawHelper() : AbortStart(MAX_uint16), AbortEnd(MAX_uint16), SearchStart(MAX_uint16), SearchEnd(MAX_uint16) {}
 };
 
-USTRUCT()
-struct FClassData
-{
-	GENERATED_USTRUCT_BODY()
-
-	FClassData() {}
-	FClassData(UClass* InClass, const FString& InDeprecatedMessage);
-	FClassData(const FString& InAssetName, const FString& InGeneratedClassPackage, const FString& InClassName, UClass* InClass);
-
-	FString ToString() const;
-	FString GetClassName() const;
-	FString GetCategory() const;
-	UClass* GetClass(bool bSilent = false);
-	bool IsAbstract() const;
-	
-	FORCEINLINE bool IsBlueprint() const { return AssetName.Len() > 0; }
-	FORCEINLINE bool IsDeprecated() const { return DeprecatedMessage.Len() > 0; }
-	FORCEINLINE FString GetDeprecatedMessage() const { return DeprecatedMessage; }
-	FORCEINLINE FString GetPackageName() const { return GeneratedClassPackage; }
-
-	/** set when child class masked this one out (e.g. always use game specific class instead of engine one) */
-	uint32 bIsHidden : 1;
-
-	/** set when class wants to hide parent class from selection (just one class up hierarchy) */
-	uint32 bHideParent : 1;
-
-private:
-
-	/** pointer to uclass */
-	TWeakObjectPtr<UClass> Class;
-
-	/** path to class if it's not loaded yet */
-	UPROPERTY()
-	FString AssetName;
-	UPROPERTY()
-	FString GeneratedClassPackage;
-
-	/** resolved name of class from asset data */
-	UPROPERTY()
-	FString ClassName;
-
-	/** User-defined category for this class */
-	UPROPERTY()
-	FString Category;
-
-	/** message for deprecated class */
-	FString DeprecatedMessage;
-};
-
-struct FClassDataNode
-{
-	FClassData Data;
-	FString ParentClassName;
-
-	TSharedPtr<FClassDataNode> ParentNode;
-	TArray<TSharedPtr<FClassDataNode> > SubNodes;
-
-	void AddUniqueSubNode(TSharedPtr<FClassDataNode> SubNode);
-};
-
-struct FClassBrowseHelper
-{
-	DECLARE_MULTICAST_DELEGATE(FOnPackageListUpdated);
-
-	FClassBrowseHelper();
-	~FClassBrowseHelper();
-
-	static void GatherClasses(const UClass* BaseClass, TArray<FClassData>& AvailableClasses);
-	static FString GetDeprecationMessage(const UClass* Class);
-
-	void OnAssetAdded(const class FAssetData& AssetData);
-	void OnAssetRemoved(const class FAssetData& AssetData);
-	void InvalidateCache();
-	void OnHotReload( bool bWasTriggeredAutomatically );
-	
-	static void AddUnknownClass(const FClassData& ClassData);
-	static bool IsClassKnown(const FClassData& ClassData);
-	static FOnPackageListUpdated OnPackageListUpdated;
-
-	/** Whether we have more than one task class available */
-	static bool IsMoreThanOneTaskClassAvailable();
-
-	/** Whether we have more than one decorator class available */
-	static bool IsMoreThanOneDecoratorClassAvailable();
-
-	/** Whether we have more than one service class available */
-	static bool IsMoreThanOneServiceClassAvailable();
-
-private:
-
-	TSharedPtr<FClassDataNode> RootNode;
-	static TArray<FName> UnknownPackages;
-	
-	TSharedPtr<FClassDataNode> CreateClassDataNode(const class FAssetData& AssetData);
-	TSharedPtr<FClassDataNode> FindBaseClassNode(TSharedPtr<FClassDataNode> Node, const FString& ClassName);
-	void FindAllSubClasses(TSharedPtr<FClassDataNode> Node, TArray<FClassData>& AvailableClasses);
-
-	UClass* FindAssetClass(const FString& GeneratedClassPackage, const FString& AssetName);
-	void BuildClassGraph();
-	void AddClassGraphChildren(TSharedPtr<FClassDataNode> Node, TArray<TSharedPtr<FClassDataNode> >& NodeList);
-
-	bool IsHidingClass(UClass* Class);
-	bool IsHidingParentClass(UClass* Class);
-	bool IsPackageSaved(FName PackageName);
-
-	/** Update the flags that determine whether to allow creation of new nodes from multiple base classes */
-	void UpdateAvailableNodeClasses();
-
-	/** Cached flag to determine whether we have more than one task class available */
-	static bool bMoreThanOneTaskClassAvailable;
-
-	/** Cached flag to determine whether we have more than one decorator class available */
-	static bool bMoreThanOneDecoratorClassAvailable;
-
-	/** Cached flag to determine whether we have more than one service class available */
-	static bool bMoreThanOneServiceClassAvailable;
-};
-
 struct FCompareNodeXLocation
 {
 	FORCEINLINE bool operator()(const UEdGraphPin& A, const UEdGraphPin& B) const
diff --git a/Engine/Source/Editor/BehaviorTreeEditor/Classes/BehaviorTreeGraph.h b/Engine/Source/Editor/BehaviorTreeEditor/Classes/BehaviorTreeGraph.h
--- a/Engine/Source/Editor/BehaviorTreeEditor/Classes/BehaviorTreeGraph.h
+++ b/Engine/Source/Editor/BehaviorTreeEditor/Classes/BehaviorTreeGraph.h
@@ -1,48 +1,50 @@
 // Copyright 1998-2015 Epic Games, Inc. All Rights Reserved.
 
 #pragma once
+
+#include "AIGraph.h"
 #include "BehaviorTreeGraph.generated.h"
 
 UCLASS()
-class UBehaviorTreeGraph : public UEdGraph
+class UBehaviorTreeGraph : public UAIGraph
 {
 	GENERATED_UCLASS_BODY()
 
-	enum EDebuggerFlags
+	enum EUpdateFlags
 	{
-		ClearDebuggerFlags,
-		SkipDebuggerFlags,
+		RebuildGraph = 0,
+		ClearDebuggerFlags = 1,
+		KeepRebuildCounter = 2,
 	};
 
-	/** counter incremented on every graph update, used for refreshing injected nodes in parent graph */
+	/** increased with every graph rebuild, used to refresh data from subtrees */
+	UPROPERTY()
+	int32 ModCounter;
+
 	UPROPERTY()
-	int32 GraphVersion;
+	bool bIsUsingModCounter;
+
+	virtual void UpdateVersion() override;
+	virtual void MarkVersion() override;
+	virtual void UpdateAsset(int32 UpdateFlags = 0) override;
+	virtual void OnSubNodeDropped() override;
 
 	void UpdateBlackboardChange();
-	void UpdateAsset(EDebuggerFlags DebuggerFlags, bool bBumpVersion = true);
 	void UpdateAbortHighlight(struct FAbortDrawHelper& Mode0, struct FAbortDrawHelper& Mode1);
 	void CreateBTFromGraph(class UBehaviorTreeGraphNode* RootEdNode);
 	void SpawnMissingNodes();
-	void RemoveOrphanedNodes();
 	void UpdatePinConnectionTypes();
 	void UpdateDeprecatedNodes();
 	bool UpdateInjectedNodes();
-	bool UpdateUnknownNodeClasses();
 	class UEdGraphNode* FindInjectedNode(int32 Index);
 	void ReplaceNodeConnections(UEdGraphNode* OldNode, UEdGraphNode* NewNode);
 	void RebuildExecutionOrder();
 
-	void LockUpdates();
-	void UnlockUpdates();
-
 	void AutoArrange();
 
 protected:
 
-	/** if set, graph modifications won't cause updates in internal tree structure (skipping UpdateAsset)
-	 * flag allows freezing update during heavy changes like pasting new nodes */
-	uint32 bLockUpdates : 1;
-};
-
-
+	virtual void CollectAllNodeInstances(TSet<UObject*>& NodeInstances);
 
+	void UpdateVersion_UnifiedSubNodes();
+};
diff --git a/Engine/Source/Editor/BehaviorTreeEditor/Classes/BehaviorTreeGraphNode.h b/Engine/Source/Editor/BehaviorTreeEditor/Classes/BehaviorTreeGraphNode.h
--- a/Engine/Source/Editor/BehaviorTreeEditor/Classes/BehaviorTreeGraphNode.h
+++ b/Engine/Source/Editor/BehaviorTreeEditor/Classes/BehaviorTreeGraphNode.h
@@ -1,22 +1,16 @@
 // Copyright 1998-2015 Epic Games, Inc. All Rights Reserved.
 
 #pragma once
+
+#include "AIGraphNode.h"
 #include "BehaviorTreeEditorTypes.h"
-#include "DiffResults.h"
 #include "BehaviorTreeGraphNode.generated.h"
 
 UCLASS()
-class UBehaviorTreeGraphNode : public UEdGraphNode
+class UBehaviorTreeGraphNode : public UAIGraphNode
 {
 	GENERATED_UCLASS_BODY()
 
-	UPROPERTY()
-	UObject* NodeInstance;
-
-	/** parent node if node is not a standalone node*/
-	UPROPERTY(transient)
-	UBehaviorTreeGraphNode* ParentNode;
-
 	/** only some of behavior tree nodes support decorators */
 	UPROPERTY()
 	TArray<UBehaviorTreeGraphNode*> Decorators;
@@ -28,71 +22,27 @@ class UBehaviorTreeGraphNode : public UEdGraphNode
 	// Begin UEdGraphNode Interface
 	virtual class UBehaviorTreeGraph* GetBehaviorTreeGraph();
 	virtual void AllocateDefaultPins() override;
-	virtual void AutowireNewNode(UEdGraphPin* FromPin) override;
-	virtual void PostPlacedNewNode() override;
-	virtual void PrepareForCopying() override;
-	virtual bool CanDuplicateNode() const override;
-	virtual bool CanUserDeleteNode() const override;
-	virtual void DestroyNode() override;
 	virtual FText GetTooltipText() const override;
-	virtual void NodeConnectionListChanged() override;
 	virtual bool CanCreateUnderSpecifiedSchema(const UEdGraphSchema* DesiredSchema) const override;
-	virtual void FindDiffs(class UEdGraphNode* OtherNode, struct FDiffResults& Results) override;
 	// End UEdGraphNode Interface
 
-	// Begin UObject Interface
-#if WITH_EDITOR
-	virtual void PostEditImport() override;
-	virtual void PostEditUndo() override;
-#endif
-	// End UObject
-
-	// @return the input pin for this state
-	virtual UEdGraphPin* GetInputPin(int32 InputIndex=0) const;
-	// @return the output pin for this state
-	virtual UEdGraphPin* GetOutputPin(int32 InputIndex=0) const;
-	//
-	virtual UEdGraph* GetBoundGraph() const { return NULL; }
-
-	virtual FString	GetDescription() const;
+	virtual FText GetDescription() const override;
+	virtual bool HasErrors() const override;
+	virtual void InitializeInstance() override;
+	virtual void OnSubNodeAdded(UAIGraphNode* SubNode) override;
+	virtual void OnSubNodeRemoved(UAIGraphNode* SubNode) override;
+	virtual void RemoveAllSubNodes();
+	virtual int32 FindSubNodeDropIndex(UAIGraphNode* SubNode) const override;
+	virtual void InsertSubNodeAt(UAIGraphNode* SubNode, int32 DropIndex) override;
+
 	/** check if node can accept breakpoints */
 	virtual bool CanPlaceBreakpoints() const { return false; }
-	virtual void PostCopyNode();
-
-	virtual bool IsSubNode() const { return false; }
-
-	/**
-	 * Finds the difference in Behavior Tree properties
-	 * 
-	 * @param Struct The struct of the class we are looking at
-	 * @param Data The raw data for the UObject we are comparing LHS
-	 * @param OtherData The raw data for the UObject we are comparing RHS
-	 * @param Results The Results where differences are stored
-	 * @param Diff The single result with default parameters setup
-	 */
-	void DiffProperties(UStruct* Struct, void* Data, void* OtherData, FDiffResults& Results, FDiffSingleResult& Diff);
 
 	void ClearDebuggerState();
 
-	/** add subnode */
-	void AddSubNode(UBehaviorTreeGraphNode* NodeTemplate, class UEdGraph* ParentGraph);
-
-	/** reinitialize node instance */
-	virtual bool RefreshNodeClass();
-
 	/** gets icon resource name for title bar */
 	virtual FName GetNameIcon() const;
 
-	/** Check if node instance uses blueprint for its implementation */
-	bool UsesBlueprint() const;
-
-	/** check if node has any errors, used for assigning colors on graph */
-	virtual bool HasErrors() const;
-
-	/** instance class */
-	UPROPERTY()
-	struct FClassData ClassData;
-
 	/** if set, this node was injected from subtree and shouldn't be edited */
 	UPROPERTY()
 	uint32 bInjectedNode : 1;
@@ -166,13 +116,6 @@ class UBehaviorTreeGraphNode : public UEdGraphNode
 	/** used to show node's runtime description rather than static one */
 	FString DebuggerRuntimeDescription;
 
-	/** error message for node */
-	FString ErrorMessage;
-
-	/** subnode index assigned during copy operation to connect nodes again on paste */
-	UPROPERTY()
-	int32 CopySubNodeIndex;
-
 protected:
 
 	/** creates add decorator... submenu */
@@ -186,7 +129,4 @@ protected:
 
 	/** add right click menu to create subnodes: Services */
 	void AddContextMenuActionsServices(const FGraphNodeContextMenuBuilder& Context) const;
-
-	void ResetNodeOwner();
 };
-
diff --git a/Engine/Source/Editor/BehaviorTreeEditor/Classes/BehaviorTreeGraphNode_CompositeDecorator.h b/Engine/Source/Editor/BehaviorTreeEditor/Classes/BehaviorTreeGraphNode_CompositeDecorator.h
--- a/Engine/Source/Editor/BehaviorTreeEditor/Classes/BehaviorTreeGraphNode_CompositeDecorator.h
+++ b/Engine/Source/Editor/BehaviorTreeEditor/Classes/BehaviorTreeGraphNode_CompositeDecorator.h
@@ -26,7 +26,7 @@ class UBehaviorTreeGraphNode_CompositeDecorator : public UBehaviorTreeGraphNode
 	FString GetNodeTypeDescription() const;
 	virtual FText GetNodeTitle(ENodeTitleType::Type TitleType) const override;
 	virtual void AllocateDefaultPins() override;
-	virtual FString GetDescription() const override;
+	virtual FText GetDescription() const override;
 	virtual void PostPlacedNewNode() override;
 	virtual void PostLoad() override;
 	virtual UEdGraph* GetBoundGraph() const { return BoundGraph; }
diff --git a/Engine/Source/Editor/BehaviorTreeEditor/Classes/BehaviorTreeGraphNode_Decorator.h b/Engine/Source/Editor/BehaviorTreeEditor/Classes/BehaviorTreeGraphNode_Decorator.h
--- a/Engine/Source/Editor/BehaviorTreeEditor/Classes/BehaviorTreeGraphNode_Decorator.h
+++ b/Engine/Source/Editor/BehaviorTreeEditor/Classes/BehaviorTreeGraphNode_Decorator.h
@@ -11,6 +11,4 @@ class UBehaviorTreeGraphNode_Decorator : public UBehaviorTreeGraphNode
 	virtual FText GetNodeTitle(ENodeTitleType::Type TitleType) const override;
 	virtual void AllocateDefaultPins() override;
 	void CollectDecoratorData(TArray<class UBTDecorator*>& NodeInstances, TArray<struct FBTDecoratorLogic>& Operations) const;
-
-	virtual bool IsSubNode() const override;
 };
diff --git a/Engine/Source/Editor/BehaviorTreeEditor/Classes/BehaviorTreeGraphNode_Root.h b/Engine/Source/Editor/BehaviorTreeEditor/Classes/BehaviorTreeGraphNode_Root.h
--- a/Engine/Source/Editor/BehaviorTreeEditor/Classes/BehaviorTreeGraphNode_Root.h
+++ b/Engine/Source/Editor/BehaviorTreeEditor/Classes/BehaviorTreeGraphNode_Root.h
@@ -25,7 +25,7 @@ class UBehaviorTreeGraphNode_Root : public UBehaviorTreeGraphNode
 
 	virtual void PostEditChangeProperty( struct FPropertyChangedEvent& PropertyChangedEvent) override;
 	virtual void PostEditUndo() override;
-	virtual FString	GetDescription() const override;
+	virtual FText GetDescription() const override;
 
 	/** notify behavior tree about blackboard change */
 	void UpdateBlackboard();
diff --git a/Engine/Source/Editor/BehaviorTreeEditor/Classes/BehaviorTreeGraphNode_Service.h b/Engine/Source/Editor/BehaviorTreeEditor/Classes/BehaviorTreeGraphNode_Service.h
--- a/Engine/Source/Editor/BehaviorTreeEditor/Classes/BehaviorTreeGraphNode_Service.h
+++ b/Engine/Source/Editor/BehaviorTreeEditor/Classes/BehaviorTreeGraphNode_Service.h
@@ -10,6 +10,5 @@ class UBehaviorTreeGraphNode_Service : public UBehaviorTreeGraphNode
 
 	virtual FText GetNodeTitle(ENodeTitleType::Type TitleType) const override;
 	virtual void AllocateDefaultPins() override;
-	virtual bool IsSubNode() const override;
 };
 
diff --git a/Engine/Source/Editor/BehaviorTreeEditor/Classes/EdGraphSchema_BehaviorTree.h b/Engine/Source/Editor/BehaviorTreeEditor/Classes/EdGraphSchema_BehaviorTree.h
--- a/Engine/Source/Editor/BehaviorTreeEditor/Classes/EdGraphSchema_BehaviorTree.h
+++ b/Engine/Source/Editor/BehaviorTreeEditor/Classes/EdGraphSchema_BehaviorTree.h
@@ -1,75 +1,9 @@
 // Copyright 1998-2015 Epic Games, Inc. All Rights Reserved.
 
 #pragma once
-#include "EdGraphSchema_BehaviorTree.generated.h"
-
-/** Action to add a node to the graph */
-USTRUCT()
-struct FBehaviorTreeSchemaAction_NewNode : public FEdGraphSchemaAction
-{
-	GENERATED_USTRUCT_BODY();
-
-	/** Template of node we want to create */
-	UPROPERTY()
-	class UBehaviorTreeGraphNode* NodeTemplate;
-
-
-	FBehaviorTreeSchemaAction_NewNode() 
-		: FEdGraphSchemaAction()
-		, NodeTemplate(NULL)
-	{}
-
-	FBehaviorTreeSchemaAction_NewNode(const FString& InNodeCategory, const FText& InMenuDesc, const FString& InToolTip, const int32 InGrouping)
-		: FEdGraphSchemaAction(InNodeCategory, InMenuDesc, InToolTip, InGrouping) 
-		, NodeTemplate(NULL)
-	{}
-
-	// FEdGraphSchemaAction interface
-	virtual UEdGraphNode* PerformAction(class UEdGraph* ParentGraph, UEdGraphPin* FromPin, const FVector2D Location, bool bSelectNewNode = true) override;
-	virtual UEdGraphNode* PerformAction(class UEdGraph* ParentGraph, TArray<UEdGraphPin*>& FromPins, const FVector2D Location, bool bSelectNewNode = true) override;
-	virtual void AddReferencedObjects( FReferenceCollector& Collector ) override;
-	// End of FEdGraphSchemaAction interface
-
-	template <typename NodeType>
-	static NodeType* SpawnNodeFromTemplate(class UEdGraph* ParentGraph, NodeType* InTemplateNode, const FVector2D Location)
-	{
-		FBehaviorTreeSchemaAction_NewNode Action;
-		Action.NodeTemplate = InTemplateNode;
-
-		return Cast<NodeType>(Action.PerformAction(ParentGraph, NULL, Location));
-	}
-};
 
-/** Action to add a subnode to the selected node */
-USTRUCT()
-struct FBehaviorTreeSchemaAction_NewSubNode : public FEdGraphSchemaAction
-{
-	GENERATED_USTRUCT_BODY();
-
-	/** Template of node we want to create */
-	UPROPERTY()
-	class UBehaviorTreeGraphNode* NodeTemplate;
-
-	/** parent node */
-	UPROPERTY()
-	class UBehaviorTreeGraphNode* ParentNode;
-
-	FBehaviorTreeSchemaAction_NewSubNode() 
-		: FEdGraphSchemaAction()
-		, NodeTemplate(NULL)
-	{}
-
-	FBehaviorTreeSchemaAction_NewSubNode(const FString& InNodeCategory, const FText& InMenuDesc, const FString& InToolTip, const int32 InGrouping)
-		: FEdGraphSchemaAction(InNodeCategory, InMenuDesc, InToolTip, InGrouping) 
-		, NodeTemplate(NULL)
-	{}
-
-	// FEdGraphSchemaAction interface
-	virtual UEdGraphNode* PerformAction(class UEdGraph* ParentGraph, UEdGraphPin* FromPin, const FVector2D Location, bool bSelectNewNode = true) override;
-	virtual UEdGraphNode* PerformAction(class UEdGraph* ParentGraph, TArray<UEdGraphPin*>& FromPins, const FVector2D Location, bool bSelectNewNode = true) override;
-	virtual void AddReferencedObjects( FReferenceCollector& Collector ) override;
-	// End of FEdGraphSchemaAction interface
-};
+#include "AIGraphSchema.h"
+#include "EdGraphSchema_BehaviorTree.generated.h"
 
 /** Action to auto arrange the graph */
 USTRUCT()
@@ -90,12 +24,10 @@ struct FBehaviorTreeSchemaAction_AutoArrange : public FEdGraphSchemaAction
 };
 
 UCLASS(MinimalAPI)
-class UEdGraphSchema_BehaviorTree : public UEdGraphSchema
+class UEdGraphSchema_BehaviorTree : public UAIGraphSchema
 {
 	GENERATED_UCLASS_BODY()
 
-	void GetBreakLinkToSubMenuActions(class FMenuBuilder& MenuBuilder, class UEdGraphPin* InGraphPin);
-
 	// Begin EdGraphSchema interface
 	virtual void CreateDefaultNodesForGraph(UEdGraph& Graph) const override;
 	virtual void GetGraphContextActions(FGraphContextMenuBuilder& ContextMenuBuilder) const override;
@@ -103,18 +35,11 @@ class UEdGraphSchema_BehaviorTree : public UEdGraphSchema
 	virtual const FPinConnectionResponse CanCreateConnection(const UEdGraphPin* A, const UEdGraphPin* B) const override;
 	virtual const FPinConnectionResponse CanMergeNodes(const UEdGraphNode* A, const UEdGraphNode* B) const override;
 	virtual FLinearColor GetPinTypeColor(const FEdGraphPinType& PinType) const override;
-	virtual bool ShouldHidePinDefaultValue(UEdGraphPin* Pin) const override;
 	virtual class FConnectionDrawingPolicy* CreateConnectionDrawingPolicy(int32 InBackLayerID, int32 InFrontLayerID, float InZoomFactor, const FSlateRect& InClippingRect, class FSlateWindowElementList& InDrawElements, class UEdGraph* InGraphObj) const override;
 	virtual int32 GetNodeSelectionCount(const UEdGraph* Graph) const override;
-	virtual void BreakNodeLinks(UEdGraphNode& TargetNode) const override;
-	virtual void BreakPinLinks(UEdGraphPin& TargetPin, bool bSendsNodeNotification) const override;
-	virtual void BreakSinglePinLink(UEdGraphPin* SourcePin, UEdGraphPin* TargetPin) override;
 	// End EdGraphSchema interface
 
-	/* subnode creation context menu */
-	void GetGraphNodeContextActions(FGraphContextMenuBuilder& ContextMenuBuilder, ESubNode::Type SubNodeType) const;
-
-protected:
-	FString GetShortTypeName(const UObject* Ob) const;
+	virtual void GetGraphNodeContextActions(FGraphContextMenuBuilder& ContextMenuBuilder, int32 SubNodeFlags) const override;
+	virtual void GetSubNodeClasses(int32 SubNodeFlags, TArray<FGraphNodeClassData>& ClassData, UClass*& GraphNodeClass) const override;
 };
 
diff --git a/Engine/Source/Editor/BehaviorTreeEditor/Private/BehaviorTreeConnectionDrawingPolicy.cpp b/Engine/Source/Editor/BehaviorTreeEditor/Private/BehaviorTreeConnectionDrawingPolicy.cpp
--- a/Engine/Source/Editor/BehaviorTreeEditor/Private/BehaviorTreeConnectionDrawingPolicy.cpp
+++ b/Engine/Source/Editor/BehaviorTreeEditor/Private/BehaviorTreeConnectionDrawingPolicy.cpp
@@ -6,9 +6,8 @@
 #include "BehaviorTreeConnectionDrawingPolicy.h"
 #include "BehaviorTreeColors.h"
 
-FBehaviorTreeConnectionDrawingPolicy::FBehaviorTreeConnectionDrawingPolicy( int32 InBackLayerID, int32 InFrontLayerID, float ZoomFactor, const FSlateRect& InClippingRect, FSlateWindowElementList& InDrawElements, UEdGraph* InGraphObj )
-	: FConnectionDrawingPolicy(InBackLayerID, InFrontLayerID, ZoomFactor, InClippingRect, InDrawElements)
-	, GraphObj(InGraphObj)
+FBehaviorTreeConnectionDrawingPolicy::FBehaviorTreeConnectionDrawingPolicy(int32 InBackLayerID, int32 InFrontLayerID, float ZoomFactor, const FSlateRect& InClippingRect, FSlateWindowElementList& InDrawElements, UEdGraph* InGraphObj)
+	: FAIGraphConnectionDrawingPolicy(InBackLayerID, InFrontLayerID, ZoomFactor, InClippingRect, InDrawElements, InGraphObj)
 {
 }
 
@@ -64,100 +63,3 @@ void FBehaviorTreeConnectionDrawingPolicy::DetermineWiringStyle(UEdGraphPin* Out
 		}
 	}
 }
-
-void FBehaviorTreeConnectionDrawingPolicy::Draw(TMap<TSharedRef<SWidget>, FArrangedWidget>& PinGeometries, FArrangedChildren& ArrangedNodes)
-{
-	// Build an acceleration structure to quickly find geometry for the nodes
-	NodeWidgetMap.Empty();
-	for (int32 NodeIndex = 0; NodeIndex < ArrangedNodes.Num(); ++NodeIndex)
-	{
-		FArrangedWidget& CurWidget = ArrangedNodes[NodeIndex];
-		TSharedRef<SGraphNode> ChildNode = StaticCastSharedRef<SGraphNode>(CurWidget.Widget);
-		NodeWidgetMap.Add(ChildNode->GetNodeObj(), NodeIndex);
-	}
-
-	// Now draw
-	FConnectionDrawingPolicy::Draw(PinGeometries, ArrangedNodes);
-}
-
-void FBehaviorTreeConnectionDrawingPolicy::DrawPreviewConnector(const FGeometry& PinGeometry, const FVector2D& StartPoint, const FVector2D& EndPoint, UEdGraphPin* Pin)
-{
-	FConnectionParams Params;
-	DetermineWiringStyle(Pin, NULL, /*inout*/ Params);
-
-	if (Pin->Direction == EEdGraphPinDirection::EGPD_Output)
-	{
-		DrawSplineWithArrow(FGeometryHelper::FindClosestPointOnGeom(PinGeometry, EndPoint), EndPoint, Params);
-	}
-	else
-	{
-		DrawSplineWithArrow(FGeometryHelper::FindClosestPointOnGeom(PinGeometry, StartPoint), StartPoint, Params);
-	}
-}
-
-void FBehaviorTreeConnectionDrawingPolicy::DrawSplineWithArrow(const FVector2D& StartAnchorPoint, const FVector2D& EndAnchorPoint, const FConnectionParams& Params)
-{
-	// bUserFlag1 indicates that we need to reverse the direction of connection (used by debugger)
-	const FVector2D& P0 = Params.bUserFlag1 ? EndAnchorPoint : StartAnchorPoint;
-	const FVector2D& P1 = Params.bUserFlag1 ? StartAnchorPoint : EndAnchorPoint;
-
-	Internal_DrawLineWithArrow(P0, P1, Params);
-}
-
-void FBehaviorTreeConnectionDrawingPolicy::Internal_DrawLineWithArrow(const FVector2D& StartAnchorPoint, const FVector2D& EndAnchorPoint, const FConnectionParams& Params)
-{
-	//@TODO: Should this be scaled by zoom factor?
-	const float LineSeparationAmount = 4.5f;
-
-	const FVector2D DeltaPos = EndAnchorPoint - StartAnchorPoint;
-	const FVector2D UnitDelta = DeltaPos.GetSafeNormal();
-	const FVector2D Normal = FVector2D(DeltaPos.Y, -DeltaPos.X).GetSafeNormal();
-
-	// Come up with the final start/end points
-	const FVector2D DirectionBias = Normal * LineSeparationAmount;
-	const FVector2D LengthBias = ArrowRadius.X * UnitDelta;
-	const FVector2D StartPoint = StartAnchorPoint + DirectionBias + LengthBias;
-	const FVector2D EndPoint = EndAnchorPoint + DirectionBias - LengthBias;
-
-	// Draw a line/spline
-	DrawConnection(WireLayerID, StartPoint, EndPoint, Params);
-
-	// Draw the arrow
-	const FVector2D ArrowDrawPos = EndPoint - ArrowRadius;
-	const float AngleInRadians = FMath::Atan2(DeltaPos.Y, DeltaPos.X);
-
-	FSlateDrawElement::MakeRotatedBox(
-		DrawElementsList,
-		ArrowLayerID,
-		FPaintGeometry(ArrowDrawPos, ArrowImage->ImageSize * ZoomFactor, ZoomFactor),
-		ArrowImage,
-		ClippingRect,
-		ESlateDrawEffect::None,
-		AngleInRadians,
-		TOptional<FVector2D>(),
-		FSlateDrawElement::RelativeToElement,
-		Params.WireColor
-		);
-}
-
-void FBehaviorTreeConnectionDrawingPolicy::DrawSplineWithArrow(const FGeometry& StartGeom, const FGeometry& EndGeom, const FConnectionParams& Params)
-{
-	// Get a reasonable seed point (halfway between the boxes)
-	const FVector2D StartCenter = FGeometryHelper::CenterOf(StartGeom);
-	const FVector2D EndCenter = FGeometryHelper::CenterOf(EndGeom);
-	const FVector2D SeedPoint = (StartCenter + EndCenter) * 0.5f;
-
-	// Find the (approximate) closest points between the two boxes
-	const FVector2D StartAnchorPoint = FGeometryHelper::FindClosestPointOnGeom(StartGeom, SeedPoint);
-	const FVector2D EndAnchorPoint = FGeometryHelper::FindClosestPointOnGeom(EndGeom, SeedPoint);
-
-	DrawSplineWithArrow(StartAnchorPoint, EndAnchorPoint, Params);
-}
-
-FVector2D FBehaviorTreeConnectionDrawingPolicy::ComputeSplineTangent(const FVector2D& Start, const FVector2D& End) const
-{
-	const FVector2D Delta = End - Start;
-	const FVector2D NormDelta = Delta.GetSafeNormal();
-
-	return NormDelta;
-}
diff --git a/Engine/Source/Editor/BehaviorTreeEditor/Private/BehaviorTreeConnectionDrawingPolicy.h b/Engine/Source/Editor/BehaviorTreeEditor/Private/BehaviorTreeConnectionDrawingPolicy.h
--- a/Engine/Source/Editor/BehaviorTreeEditor/Private/BehaviorTreeConnectionDrawingPolicy.h
+++ b/Engine/Source/Editor/BehaviorTreeEditor/Private/BehaviorTreeConnectionDrawingPolicy.h
@@ -1,27 +1,16 @@
 // Copyright 1998-2015 Epic Games, Inc. All Rights Reserved.
 
 #pragma once
+#include "AIGraphConnectionDrawingPolicy.h"
 
 // This class draws the connections for an UEdGraph with a behavior tree schema
-class BEHAVIORTREEEDITOR_API FBehaviorTreeConnectionDrawingPolicy : public FConnectionDrawingPolicy
+class BEHAVIORTREEEDITOR_API FBehaviorTreeConnectionDrawingPolicy : public FAIGraphConnectionDrawingPolicy
 {
-protected:
-	UEdGraph* GraphObj;
-
-	TMap<UEdGraphNode*, int32> NodeWidgetMap;
 public:
 	//
 	FBehaviorTreeConnectionDrawingPolicy(int32 InBackLayerID, int32 InFrontLayerID, float ZoomFactor, const FSlateRect& InClippingRect, FSlateWindowElementList& InDrawElements, UEdGraph* InGraphObj);
 
 	// FConnectionDrawingPolicy interface 
 	virtual void DetermineWiringStyle(UEdGraphPin* OutputPin, UEdGraphPin* InputPin, /*inout*/ FConnectionParams& Params) override;
-	virtual void Draw(TMap<TSharedRef<SWidget>, FArrangedWidget>& PinGeometries, FArrangedChildren& ArrangedNodes) override;
-	virtual void DrawSplineWithArrow(const FGeometry& StartGeom, const FGeometry& EndGeom, const FConnectionParams& Params) override;
-	virtual void DrawSplineWithArrow(const FVector2D& StartPoint, const FVector2D& EndPoint, const FConnectionParams& Params) override;
-	virtual void DrawPreviewConnector(const FGeometry& PinGeometry, const FVector2D& StartPoint, const FVector2D& EndPoint, UEdGraphPin* Pin) override;
-	virtual FVector2D ComputeSplineTangent(const FVector2D& Start, const FVector2D& End) const override;
 	// End of FConnectionDrawingPolicy interface
-
-protected:
-	void Internal_DrawLineWithArrow(const FVector2D& StartAnchorPoint, const FVector2D& EndAnchorPoint, const FConnectionParams& Params);
 };
diff --git a/Engine/Source/Editor/BehaviorTreeEditor/Private/BehaviorTreeDecoratorGraphNode_Decorator.cpp b/Engine/Source/Editor/BehaviorTreeEditor/Private/BehaviorTreeDecoratorGraphNode_Decorator.cpp
--- a/Engine/Source/Editor/BehaviorTreeEditor/Private/BehaviorTreeDecoratorGraphNode_Decorator.cpp
+++ b/Engine/Source/Editor/BehaviorTreeEditor/Private/BehaviorTreeDecoratorGraphNode_Decorator.cpp
@@ -18,7 +18,7 @@ void UBehaviorTreeDecoratorGraphNode_Decorator::AllocateDefaultPins()
 
 void UBehaviorTreeDecoratorGraphNode_Decorator::PostPlacedNewNode()
 {
-	UClass* NodeClass = ClassData.GetClass();
+	UClass* NodeClass = ClassData.GetClass(true);
 	if (NodeClass != NULL)
 	{
 		UBehaviorTreeGraphNode_CompositeDecorator* OwningNode = Cast<UBehaviorTreeGraphNode_CompositeDecorator>(GetDecoratorGraph()->GetOuter());
@@ -105,14 +105,14 @@ bool UBehaviorTreeDecoratorGraphNode_Decorator::RefreshNodeClass()
 	bool bUpdated = false;
 	if (NodeInstance == NULL)
 	{
-		if (FClassBrowseHelper::IsClassKnown(ClassData))
+		if (FGraphNodeClassHelper::IsClassKnown(ClassData))
 		{
 			PostPlacedNewNode();
 			bUpdated = (NodeInstance != NULL);
 		}
 		else
 		{
-			FClassBrowseHelper::AddUnknownClass(ClassData);
+			FGraphNodeClassHelper::AddUnknownClass(ClassData);
 		}
 	}
 
diff --git a/Engine/Source/Editor/BehaviorTreeEditor/Private/BehaviorTreeEditor.cpp b/Engine/Source/Editor/BehaviorTreeEditor/Private/BehaviorTreeEditor.cpp
--- a/Engine/Source/Editor/BehaviorTreeEditor/Private/BehaviorTreeEditor.cpp
+++ b/Engine/Source/Editor/BehaviorTreeEditor/Private/BehaviorTreeEditor.cpp
@@ -45,21 +45,18 @@ const FName FBehaviorTreeEditor::BlackboardMode(TEXT("Blackboard"));
 FBehaviorTreeEditor::FBehaviorTreeEditor() 
 	: IBehaviorTreeEditor()
 {
-	UEditorEngine* Editor = (UEditorEngine*)GEngine;
-	if (Editor != NULL)
-	{
-		Editor->RegisterForUndo(this);
-	}
-
 	// listen for package change events to update injected nodes
 	OnPackageSavedDelegateHandle     = UPackage::PackageSavedEvent.AddRaw(this, &FBehaviorTreeEditor::OnPackageSaved);
-	OnClassListUpdatedDelegateHandle = FClassBrowseHelper::OnPackageListUpdated.AddRaw(this, &FBehaviorTreeEditor::OnClassListUpdated);
 
 	bShowDecoratorRangeLower = false;
 	bShowDecoratorRangeSelf = false;
 	bSelectedNodeIsInjected = false;
 	SelectedNodesCount = 0;
 
+	bHasMultipleTaskBP = false;
+	bHasMultipleDecoratorBP = false;
+	bHasMultipleServiceBP = false;
+
 	BehaviorTree = nullptr;
 	BlackboardData = nullptr;
 
@@ -68,62 +65,45 @@ FBehaviorTreeEditor::FBehaviorTreeEditor()
 
 FBehaviorTreeEditor::~FBehaviorTreeEditor()
 {
-	UEditorEngine* Editor = (UEditorEngine*)GEngine;
-	if (Editor)
-	{
-		Editor->UnregisterForUndo( this );
-	}
-
 	UPackage::PackageSavedEvent.Remove(OnPackageSavedDelegateHandle);
-	FClassBrowseHelper::OnPackageListUpdated.Remove(OnClassListUpdatedDelegateHandle);
 
 	Debugger.Reset();
 }
 
 void FBehaviorTreeEditor::PostUndo(bool bSuccess)
-{	
+{
 	if (bSuccess)
 	{
-		// Clear selection, to avoid holding refs to nodes that go away
-		TSharedPtr<SGraphEditor> CurrentGraphEditor = FocusedGraphEdPtr.Pin();
-		if (CurrentGraphEditor.IsValid())
-		{
-			CurrentGraphEditor->ClearSelectionSet();
-			CurrentGraphEditor->NotifyGraphChanged();
-		}
 		if(BlackboardView.IsValid())
 		{
 			BlackboardView->SetObject(GetBlackboardData());
 		}
+
 		if(BlackboardEditor.IsValid())
 		{
 			BlackboardEditor->SetObject(GetBlackboardData());
 		}
-		FSlateApplication::Get().DismissAllMenus();
 	}
+
+	FAIGraphEditor::PostUndo(bSuccess);
 }
 
 void FBehaviorTreeEditor::PostRedo(bool bSuccess)
 {
 	if (bSuccess)
 	{
-		// Clear selection, to avoid holding refs to nodes that go away
-		TSharedPtr<SGraphEditor> CurrentGraphEditor = FocusedGraphEdPtr.Pin();
-		if (CurrentGraphEditor.IsValid())
-		{
-			CurrentGraphEditor->ClearSelectionSet();
-			CurrentGraphEditor->NotifyGraphChanged();
-		}
-		if(BlackboardView.IsValid())
+		if (BlackboardView.IsValid())
 		{
 			BlackboardView->SetObject(GetBlackboardData());
 		}
-		if(BlackboardEditor.IsValid())
+
+		if (BlackboardEditor.IsValid())
 		{
 			BlackboardEditor->SetObject(GetBlackboardData());
 		}
-		FSlateApplication::Get().DismissAllMenus();
 	}
+
+	FAIGraphEditor::PostRedo(bSuccess);
 }
 
 void FBehaviorTreeEditor::NotifyPostChange( const FPropertyChangedEvent& PropertyChangedEvent, UProperty* PropertyThatChanged )
@@ -270,6 +250,7 @@ void FBehaviorTreeEditor::InitBehaviorTreeEditor( const EToolkitMode::Type Mode,
 		SetCurrentMode(BlackboardMode);
 	}
 
+	OnClassListUpdated();
 	RegenerateMenusAndToolbars();
 }
 
@@ -282,7 +263,7 @@ void FBehaviorTreeEditor::RestoreBehaviorTree()
 	{
 		BehaviorTree->BTGraph = FBlueprintEditorUtils::CreateNewGraph(BehaviorTree, TEXT("Behavior Tree"), UBehaviorTreeGraph::StaticClass(), UEdGraphSchema_BehaviorTree::StaticClass());
 		MyGraph = Cast<UBehaviorTreeGraph>(BehaviorTree->BTGraph);
-		check(MyGraph);
+		MyGraph->MarkVersion();
 
 		// Initialize the behavior tree graph
 		const UEdGraphSchema* Schema = MyGraph->GetSchema();
@@ -296,6 +277,7 @@ void FBehaviorTreeEditor::RestoreBehaviorTree()
 		MyGraph->UpdateDeprecatedNodes();
 	}
 
+	MyGraph->UpdateVersion();
 	MyGraph->UpdateBlackboardChange();
 	MyGraph->UpdateInjectedNodes();
 	MyGraph->UpdateUnknownNodeClasses();
@@ -312,11 +294,11 @@ void FBehaviorTreeEditor::RestoreBehaviorTree()
 	const bool bIncreaseVersionNum = false;
 	if(bNewGraph)
 	{
-		MyGraph->UpdateAsset(UBehaviorTreeGraph::ClearDebuggerFlags, bIncreaseVersionNum);
+		MyGraph->UpdateAsset(UBehaviorTreeGraph::ClearDebuggerFlags | UBehaviorTreeGraph::KeepRebuildCounter);
 	}
 	else
 	{
-		MyGraph->UpdateAsset(UBehaviorTreeGraph::SkipDebuggerFlags, bIncreaseVersionNum);
+		MyGraph->UpdateAsset(UBehaviorTreeGraph::KeepRebuildCounter);
 		RefreshDebugger();
 	}
 
@@ -506,38 +488,7 @@ TSharedRef<SGraphEditor> FBehaviorTreeEditor::CreateGraphEditorWidget(UEdGraph*
 	
 	if (!GraphEditorCommands.IsValid())
 	{
-		GraphEditorCommands = MakeShareable( new FUICommandList );
-
-		// Editing commands
-		GraphEditorCommands->MapAction( FGenericCommands::Get().SelectAll,
-			FExecuteAction::CreateSP( this, &FBehaviorTreeEditor::SelectAllNodes ),
-			FCanExecuteAction::CreateSP( this, &FBehaviorTreeEditor::CanSelectAllNodes )
-			);
-
-		GraphEditorCommands->MapAction( FGenericCommands::Get().Delete,
-			FExecuteAction::CreateSP( this, &FBehaviorTreeEditor::DeleteSelectedNodes ),
-			FCanExecuteAction::CreateSP( this, &FBehaviorTreeEditor::CanDeleteNodes )
-			);
-
-		GraphEditorCommands->MapAction( FGenericCommands::Get().Copy,
-			FExecuteAction::CreateSP( this, &FBehaviorTreeEditor::CopySelectedNodes ),
-			FCanExecuteAction::CreateSP( this, &FBehaviorTreeEditor::CanCopyNodes )
-			);
-
-		GraphEditorCommands->MapAction( FGenericCommands::Get().Cut,
-			FExecuteAction::CreateSP( this, &FBehaviorTreeEditor::CutSelectedNodes ),
-			FCanExecuteAction::CreateSP( this, &FBehaviorTreeEditor::CanCutNodes )
-			);
-
-		GraphEditorCommands->MapAction( FGenericCommands::Get().Paste,
-			FExecuteAction::CreateSP( this, &FBehaviorTreeEditor::PasteNodes ),
-			FCanExecuteAction::CreateSP( this, &FBehaviorTreeEditor::CanPasteNodes )
-			);
-
-		GraphEditorCommands->MapAction( FGenericCommands::Get().Duplicate,
-			FExecuteAction::CreateSP( this, &FBehaviorTreeEditor::DuplicateNodes ),
-			FCanExecuteAction::CreateSP( this, &FBehaviorTreeEditor::CanDuplicateNodes )
-			);
+		CreateCommandList();
 
 		GraphEditorCommands->MapAction( FGraphEditorCommands::Get().RemoveExecutionPin,
 			FExecuteAction::CreateSP( this, &FBehaviorTreeEditor::OnRemoveInputPin ),
@@ -1116,7 +1067,7 @@ bool FBehaviorTreeEditor::IsPropertyEditable() const
 		return false;
 	}
 
-	TSharedPtr<SGraphEditor> FocusedGraphEd = FocusedGraphEdPtr.Pin();
+	TSharedPtr<SGraphEditor> FocusedGraphEd = UpdateGraphEdPtr.Pin();
 	return FocusedGraphEd.IsValid() && FocusedGraphEd->GetCurrentGraph() && FocusedGraphEd->GetCurrentGraph()->bEditable;
 }
 
@@ -1184,352 +1135,15 @@ void FBehaviorTreeEditor::OnPackageSaved(const FString& PackageFileName, UObject
 
 void FBehaviorTreeEditor::OnClassListUpdated()
 {
-	UBehaviorTreeGraph* MyGraph = BehaviorTree ? Cast<UBehaviorTreeGraph>(BehaviorTree->BTGraph) : NULL;
-	if (MyGraph)
-	{
-		const bool bUpdated = MyGraph->UpdateUnknownNodeClasses();
-		if (bUpdated)
-		{
-			FGraphPanelSelectionSet CurrentSelection = GetSelectedNodes();
-			OnSelectedNodesChanged(CurrentSelection);
-
-			MyGraph->UpdateAsset(UBehaviorTreeGraph::ClearDebuggerFlags);
-		}
-	}
-}
-
-FGraphPanelSelectionSet FBehaviorTreeEditor::GetSelectedNodes() const
-{
-	FGraphPanelSelectionSet CurrentSelection;
-	TSharedPtr<SGraphEditor> FocusedGraphEd = FocusedGraphEdPtr.Pin();
-	if (FocusedGraphEd.IsValid())
-	{
-		CurrentSelection = FocusedGraphEd->GetSelectedNodes();
-	}
-	return CurrentSelection;
-}
-
-void FBehaviorTreeEditor::SelectAllNodes()
-{
-	TSharedPtr<SGraphEditor> CurrentGraphEditor = FocusedGraphEdPtr.Pin();
-	if (CurrentGraphEditor.IsValid())
-	{
-		CurrentGraphEditor->SelectAllNodes();
-	}
-}
-
-bool FBehaviorTreeEditor::CanSelectAllNodes() const
-{
-	return true;
-}
-
-void FBehaviorTreeEditor::DeleteSelectedNodes()
-{
-	TSharedPtr<SGraphEditor> CurrentGraphEditor = FocusedGraphEdPtr.Pin();
-	if (!CurrentGraphEditor.IsValid())
-	{
-		return;
-	}
-
-	const FScopedTransaction Transaction( FGenericCommands::Get().Delete->GetDescription() );
-	CurrentGraphEditor->GetCurrentGraph()->Modify();
-
-	const FGraphPanelSelectionSet SelectedNodes = CurrentGraphEditor->GetSelectedNodes();
-	CurrentGraphEditor->ClearSelectionSet();
-
-	for (FGraphPanelSelectionSet::TConstIterator NodeIt( SelectedNodes ); NodeIt; ++NodeIt)
-	{
-		if (UEdGraphNode* Node = Cast<UEdGraphNode>(*NodeIt))
-		{
-			if (Node->CanUserDeleteNode())
-			{
-				Node->Modify();
-				Node->DestroyNode();
-			}
-		}
-	}
-}
-
-bool FBehaviorTreeEditor::CanDeleteNodes() const
-{
-	// If any of the nodes can be deleted then we should allow deleting
-	const FGraphPanelSelectionSet SelectedNodes = GetSelectedNodes();
-	for (FGraphPanelSelectionSet::TConstIterator SelectedIter(SelectedNodes); SelectedIter; ++SelectedIter)
-	{
-		UEdGraphNode* Node = Cast<UEdGraphNode>(*SelectedIter);
-		if (Node && Node->CanUserDeleteNode())
-		{
-			return true;
-		}
-	}
-
-	return false;
-}
-
-void FBehaviorTreeEditor::DeleteSelectedDuplicatableNodes()
-{
-	TSharedPtr<SGraphEditor> CurrentGraphEditor = FocusedGraphEdPtr.Pin();
-	if (!CurrentGraphEditor.IsValid())
-	{
-		return;
-	}
-
-	const FGraphPanelSelectionSet OldSelectedNodes = CurrentGraphEditor->GetSelectedNodes();
-	CurrentGraphEditor->ClearSelectionSet();
-
-	for (FGraphPanelSelectionSet::TConstIterator SelectedIter(OldSelectedNodes); SelectedIter; ++SelectedIter)
-	{
-		UEdGraphNode* Node = Cast<UEdGraphNode>(*SelectedIter);
-		if (Node && Node->CanDuplicateNode())
-		{
-			CurrentGraphEditor->SetNodeSelection(Node, true);
-		}
-	}
-
-	// Delete the duplicatable nodes
-	DeleteSelectedNodes();
-
-	CurrentGraphEditor->ClearSelectionSet();
-
-	for (FGraphPanelSelectionSet::TConstIterator SelectedIter(OldSelectedNodes); SelectedIter; ++SelectedIter)
-	{
-		if (UEdGraphNode* Node = Cast<UEdGraphNode>(*SelectedIter))
-		{
-			CurrentGraphEditor->SetNodeSelection(Node, true);
-		}
-	}
-}
-
-void FBehaviorTreeEditor::CutSelectedNodes()
-{
-	CopySelectedNodes();
-	DeleteSelectedDuplicatableNodes();
-}
-
-bool FBehaviorTreeEditor::CanCutNodes() const
-{
-	return CanCopyNodes() && CanDeleteNodes();
-}
-
-void FBehaviorTreeEditor::CopySelectedNodes()
-{
-	// Export the selected nodes and place the text on the clipboard
-	FGraphPanelSelectionSet SelectedNodes = GetSelectedNodes();
-	TArray<UBehaviorTreeGraphNode*> SubNodes;
-
-	FString ExportedText;
-
-	int32 CopySubNodeIndex = 0;
-	for (FGraphPanelSelectionSet::TIterator SelectedIter(SelectedNodes); SelectedIter; ++SelectedIter)
-	{
-		UBehaviorTreeGraphNode* Node = Cast<UBehaviorTreeGraphNode>(*SelectedIter);
-		
-		// skip all manually selected subnodes
-		if (Node == NULL || Node->IsSubNode() || Node->bInjectedNode)
-		{
-			SelectedIter.RemoveCurrent();
-			continue;
-		}
-
-		Node->PrepareForCopying();
-		Node->CopySubNodeIndex = CopySubNodeIndex;
-
-		// append all subnodes for selection
-		for (int32 i = 0; i < Node->Decorators.Num(); i++)
-		{
-			if (!Node->Decorators[i]->bInjectedNode)
-			{
-				Node->Decorators[i]->CopySubNodeIndex = CopySubNodeIndex;
-				SubNodes.Add(Node->Decorators[i]);
-			}
-		}
-
-		for (int32 i = 0; i < Node->Services.Num(); i++)
-		{
-			if (!Node->Services[i]->bInjectedNode)
-			{
-				Node->Services[i]->CopySubNodeIndex = CopySubNodeIndex;
-				SubNodes.Add(Node->Services[i]);
-			}
-		}
-
-		CopySubNodeIndex++;
-	}
-
-	for (int32 i = 0; i < SubNodes.Num(); i++)
-	{
-		SelectedNodes.Add(SubNodes[i]);
-		SubNodes[i]->PrepareForCopying();
-	}
-
-	FEdGraphUtilities::ExportNodesToText(SelectedNodes, ExportedText);
-	FPlatformMisc::ClipboardCopy(*ExportedText);
-
-	for (FGraphPanelSelectionSet::TIterator SelectedIter(SelectedNodes); SelectedIter; ++SelectedIter)
-	{
-		UBehaviorTreeGraphNode* Node = Cast<UBehaviorTreeGraphNode>(*SelectedIter);
-		Node->PostCopyNode();
-	}
-}
-
-bool FBehaviorTreeEditor::CanCopyNodes() const
-{
-	// If any of the nodes can be duplicated then we should allow copying
-	const FGraphPanelSelectionSet SelectedNodes = GetSelectedNodes();
-	for (FGraphPanelSelectionSet::TConstIterator SelectedIter(SelectedNodes); SelectedIter; ++SelectedIter)
-	{
-		UEdGraphNode* Node = Cast<UEdGraphNode>(*SelectedIter);
-		if (Node && Node->CanDuplicateNode())
-		{
-			return true;
-		}
-	}
-
-	return false;
-}
-
-void FBehaviorTreeEditor::PasteNodes()
-{
-	TSharedPtr<SGraphEditor> CurrentGraphEditor = FocusedGraphEdPtr.Pin();
-	if (CurrentGraphEditor.IsValid())
-	{
-		PasteNodesHere(CurrentGraphEditor->GetPasteLocation());
-	}
-}
-
-void FBehaviorTreeEditor::PasteNodesHere(const FVector2D& Location)
-{
-	TSharedPtr<SGraphEditor> CurrentGraphEditor = FocusedGraphEdPtr.Pin();
-	if (!CurrentGraphEditor.IsValid())
-	{
-		return;
-	}
+	FAIGraphEditor::OnClassListUpdated();
 
-	// Undo/Redo support
-	const FScopedTransaction Transaction( FGenericCommands::Get().Paste->GetDescription() );
-	UBehaviorTreeGraph* BTGraph = Cast<UBehaviorTreeGraph>(CurrentGraphEditor->GetCurrentGraph());
-	BTGraph->Modify();
-	BTGraph->LockUpdates();
+	const int32 NumTaskBP = FGraphNodeClassHelper::GetObservedBlueprintClassCount(UBTTask_BlueprintBase::StaticClass());
+	const int32 NumDecoratorBP = FGraphNodeClassHelper::GetObservedBlueprintClassCount(UBTDecorator_BlueprintBase::StaticClass());
+	const int32 NumServiceBP = FGraphNodeClassHelper::GetObservedBlueprintClassCount(UBTService_BlueprintBase::StaticClass());
 
-	// Clear the selection set (newly pasted stuff will be selected)
-	CurrentGraphEditor->ClearSelectionSet();
-
-	// Grab the text to paste from the clipboard.
-	FString TextToImport;
-	FPlatformMisc::ClipboardPaste(TextToImport);
-
-	// Import the nodes
-	TSet<UEdGraphNode*> PastedNodes;
-	FEdGraphUtilities::ImportNodesFromText(BTGraph, TextToImport, /*out*/ PastedNodes);
-
-	//Average position of nodes so we can move them while still maintaining relative distances to each other
-	FVector2D AvgNodePosition(0.0f,0.0f);
-
-	float NumTopLevelNodes = 0.0f;
-	for (TSet<UEdGraphNode*>::TIterator It(PastedNodes); It; ++It)
-	{
-		UBehaviorTreeGraphNode* BTNode = Cast<UBehaviorTreeGraphNode>(*It);
-		if (BTNode && !BTNode->IsSubNode())
-		{
-			AvgNodePosition.X += BTNode->NodePosX;
-			AvgNodePosition.Y += BTNode->NodePosY;
-			NumTopLevelNodes += 1.0f;
-		}
-	}
-
-	if (NumTopLevelNodes > 0.0f)
-	{
-		float InvNumNodes = 1.0f/NumTopLevelNodes;
-		AvgNodePosition.X *= InvNumNodes;
-		AvgNodePosition.Y *= InvNumNodes;
-	}
-
-	TMap<int32, UBehaviorTreeGraphNode*> ParentMap;
-	for (TSet<UEdGraphNode*>::TIterator It(PastedNodes); It; ++It)
-	{
-		UBehaviorTreeGraphNode* BTNode = Cast<UBehaviorTreeGraphNode>(*It);
-		if (BTNode && !BTNode->IsSubNode())
-		{
-			// Select the newly pasted stuff
-			CurrentGraphEditor->SetNodeSelection(BTNode, true);
-
-			BTNode->NodePosX = (BTNode->NodePosX - AvgNodePosition.X) + Location.X ;
-			BTNode->NodePosY = (BTNode->NodePosY - AvgNodePosition.Y) + Location.Y ;
-
-			BTNode->SnapToGrid(16);
-
-			// Give new node a different Guid from the old one
-			BTNode->CreateNewGuid();
-
-			BTNode->Decorators.Reset();
-			BTNode->Services.Reset();
-
-			ParentMap.Add(BTNode->CopySubNodeIndex, BTNode);
-		}
-	}
-
-	for (TSet<UEdGraphNode*>::TIterator It(PastedNodes); It; ++It)
-	{
-		UBehaviorTreeGraphNode* BTNode = Cast<UBehaviorTreeGraphNode>(*It);
-		if (BTNode && BTNode->IsSubNode())
-		{
-			BTNode->NodePosX = 0;
-			BTNode->NodePosY = 0;
-
-			// Give new node a different Guid from the old one
-			BTNode->CreateNewGuid();
-
-			// remove subnode from graph, it will be referenced from parent node
-			BTNode->DestroyNode();
-
-			BTNode->ParentNode = ParentMap.FindRef(BTNode->CopySubNodeIndex);
-			if (BTNode->ParentNode)
-			{
-				const bool bIsService = BTNode->IsA(UBehaviorTreeGraphNode_Service::StaticClass());
-				if (bIsService)
-				{
-					BTNode->ParentNode->Services.Add(BTNode);
-				}
-				else
-				{
-					BTNode->ParentNode->Decorators.Add(BTNode);
-				}
-			}
-		}
-	}
-
-	BTGraph->UnlockUpdates();
-
-	// Update UI
-	CurrentGraphEditor->NotifyGraphChanged();
-
-	BehaviorTree->PostEditChange();
-	BehaviorTree->MarkPackageDirty();
-}
-
-bool FBehaviorTreeEditor::CanPasteNodes() const
-{
-	TSharedPtr<SGraphEditor> CurrentGraphEditor = FocusedGraphEdPtr.Pin();
-	if (!CurrentGraphEditor.IsValid())
-	{
-		return false;
-	}
-
-	FString ClipboardContent;
-	FPlatformMisc::ClipboardPaste(ClipboardContent);
-
-	return FEdGraphUtilities::CanImportNodesFromText(CurrentGraphEditor->GetCurrentGraph(), ClipboardContent);
-}
-
-void FBehaviorTreeEditor::DuplicateNodes()
-{
-	CopySelectedNodes();
-	PasteNodes();
-}
-
-bool FBehaviorTreeEditor::CanDuplicateNodes() const
-{
-	return CanCopyNodes();
+	bHasMultipleTaskBP = NumTaskBP > 1;
+	bHasMultipleDecoratorBP = NumDecoratorBP > 1;
+	bHasMultipleServiceBP = NumServiceBP > 1;
 }
 
 void FBehaviorTreeEditor::OnNodeDoubleClicked(class UEdGraphNode* Node)
@@ -1548,9 +1162,10 @@ void FBehaviorTreeEditor::OnNodeDoubleClicked(class UEdGraphNode* Node)
 				ChildNodeEditor->InitializeDebuggerState(Debugger.Get());
 				
 				int32 FirstInjectedIdx = INDEX_NONE;
-				for (int32 Idx = 0; Idx < MyNode->ParentNode->Decorators.Num(); Idx++)
+				UBehaviorTreeGraphNode* MyParentNode = Cast<UBehaviorTreeGraphNode>(MyNode->ParentNode);
+				for (int32 Idx = 0; Idx < MyParentNode->Decorators.Num(); Idx++)
 				{
-					if (MyNode->ParentNode->Decorators[Idx]->bInjectedNode)
+					if (MyParentNode->Decorators[Idx]->bInjectedNode)
 					{
 						FirstInjectedIdx = Idx;
 						break;
@@ -1559,7 +1174,7 @@ void FBehaviorTreeEditor::OnNodeDoubleClicked(class UEdGraphNode* Node)
 
 				if (FirstInjectedIdx != INDEX_NONE)
 				{
-					const int32 NodeIdx = MyNode->ParentNode->Decorators.IndexOfByKey(MyNode) - FirstInjectedIdx;
+					const int32 NodeIdx = MyParentNode->Decorators.IndexOfByKey(MyNode) - FirstInjectedIdx;
 					UEdGraphNode* OtherNode = ChildNodeEditor->FindInjectedNode(NodeIdx);
 					if (OtherNode)
 					{
@@ -1622,7 +1237,7 @@ void FBehaviorTreeEditor::OnNodeDoubleClicked(class UEdGraphNode* Node)
 void FBehaviorTreeEditor::OnAddInputPin()
 {
 	FGraphPanelSelectionSet CurrentSelection;
-	TSharedPtr<SGraphEditor> FocusedGraphEd = FocusedGraphEdPtr.Pin();
+	TSharedPtr<SGraphEditor> FocusedGraphEd = UpdateGraphEdPtr.Pin();
 	if (FocusedGraphEd.IsValid())
 	{
 		CurrentSelection = FocusedGraphEd->GetSelectedNodes();
@@ -1672,7 +1287,7 @@ bool FBehaviorTreeEditor::CanAddInputPin() const
 
 void FBehaviorTreeEditor::OnRemoveInputPin()
 {
-	TSharedPtr<SGraphEditor> FocusedGraphEd = FocusedGraphEdPtr.Pin();
+	TSharedPtr<SGraphEditor> FocusedGraphEd = UpdateGraphEdPtr.Pin();
 	if (FocusedGraphEd.IsValid())
 	{
 		const FScopedTransaction Transaction( LOCTEXT("RemoveInputPin", "Remove Input Pin") );
@@ -1714,7 +1329,7 @@ bool FBehaviorTreeEditor::CanRemoveInputPin() const
 
 void FBehaviorTreeEditor::OnGraphEditorFocused(const TSharedRef<SGraphEditor>& InGraphEditor)
 {
-	FocusedGraphEdPtr = InGraphEditor;
+	UpdateGraphEdPtr = InGraphEditor;
 	FocusedGraphOwner = Cast<UBehaviorTreeGraphNode_CompositeDecorator>(InGraphEditor->GetCurrentGraph()->GetOuter());
 
 	FGraphPanelSelectionSet CurrentSelection;
@@ -1729,7 +1344,7 @@ void FBehaviorTreeEditor::SaveAsset_Execute()
 		UBehaviorTreeGraph* BTGraph = Cast<UBehaviorTreeGraph>(BehaviorTree->BTGraph);
 		if (BTGraph)
 		{
-			BTGraph->UpdateAsset(UBehaviorTreeGraph::SkipDebuggerFlags);
+			BTGraph->UpdateAsset();
 		}
 	}
 	// save it
@@ -1908,7 +1523,7 @@ void FBehaviorTreeEditor::JumpToNode(const UEdGraphNode* Node)
 
 TWeakPtr<SGraphEditor> FBehaviorTreeEditor::GetFocusedGraphPtr() const
 {
-	return FocusedGraphEdPtr;
+	return UpdateGraphEdPtr;
 }
 
 bool FBehaviorTreeEditor::CanAccessBlackboardMode() const
@@ -1945,7 +1560,7 @@ UEdGraphNode* FBehaviorTreeEditor::FindInjectedNode(int32 Index) const
 
 void FBehaviorTreeEditor::DoubleClickNode(class UEdGraphNode* Node)
 {
-	TSharedPtr<SGraphEditor> CurrentGraphEditor = FocusedGraphEdPtr.Pin();
+	TSharedPtr<SGraphEditor> CurrentGraphEditor = UpdateGraphEdPtr.Pin();
 	if (CurrentGraphEditor.IsValid())
 	{
 		CurrentGraphEditor->ClearSelectionSet();
@@ -2135,12 +1750,12 @@ void FBehaviorTreeEditor::CreateNewTask() const
 
 bool FBehaviorTreeEditor::IsNewTaskButtonVisible() const
 {
-	return !FClassBrowseHelper::IsMoreThanOneTaskClassAvailable();
+	return !bHasMultipleTaskBP;
 }
 
 bool FBehaviorTreeEditor::IsNewTaskComboVisible() const
 {
-	return FClassBrowseHelper::IsMoreThanOneTaskClassAvailable();
+	return bHasMultipleTaskBP;
 }
 
 void FBehaviorTreeEditor::CreateNewDecorator() const
@@ -2150,12 +1765,12 @@ void FBehaviorTreeEditor::CreateNewDecorator() const
 
 bool FBehaviorTreeEditor::IsNewDecoratorButtonVisible() const
 {
-	return !FClassBrowseHelper::IsMoreThanOneDecoratorClassAvailable();
+	return !bHasMultipleDecoratorBP;
 }
 
 bool FBehaviorTreeEditor::IsNewDecoratorComboVisible() const
 {
-	return FClassBrowseHelper::IsMoreThanOneDecoratorClassAvailable();
+	return bHasMultipleDecoratorBP;
 }
 
 void FBehaviorTreeEditor::CreateNewService() const
@@ -2165,12 +1780,12 @@ void FBehaviorTreeEditor::CreateNewService() const
 
 bool FBehaviorTreeEditor::IsNewServiceButtonVisible() const
 {
-	return !FClassBrowseHelper::IsMoreThanOneServiceClassAvailable();
+	return !bHasMultipleServiceBP;
 }
 
 bool FBehaviorTreeEditor::IsNewServiceComboVisible() const
 {
-	return FClassBrowseHelper::IsMoreThanOneServiceClassAvailable();
+	return bHasMultipleServiceBP;
 }
 
 void FBehaviorTreeEditor::CreateNewBlackboard()
diff --git a/Engine/Source/Editor/BehaviorTreeEditor/Private/BehaviorTreeEditor.h b/Engine/Source/Editor/BehaviorTreeEditor/Private/BehaviorTreeEditor.h
--- a/Engine/Source/Editor/BehaviorTreeEditor/Private/BehaviorTreeEditor.h
+++ b/Engine/Source/Editor/BehaviorTreeEditor/Private/BehaviorTreeEditor.h
@@ -4,13 +4,13 @@
 
 #include "IBehaviorTreeEditor.h"
 #include "Toolkits/AssetEditorToolkit.h"
-#include "EditorUndoClient.h"
+#include "AIGraphEditor.h"
 
 class UBehaviorTree;
 class UBlackboardData;
 struct FBlackboardEntry;
 
-class FBehaviorTreeEditor : public IBehaviorTreeEditor, public FEditorUndoClient, public FNotifyHook
+class FBehaviorTreeEditor : public IBehaviorTreeEditor, public FAIGraphEditor, public FNotifyHook
 {
 public:
 	FBehaviorTreeEditor();
@@ -38,8 +38,8 @@ public:
 	// End IBehaviorTreeEditor interface
 
 	// Begin FEditorUndoClient Interface
-	virtual void	PostUndo(bool bSuccess) override;
-	virtual void	PostRedo(bool bSuccess) override;
+	virtual void PostUndo(bool bSuccess) override;
+	virtual void PostRedo(bool bSuccess) override;
 	// End of FEditorUndoClient
 
 	// Begin FNotifyHook Interface
@@ -49,7 +49,6 @@ public:
 	// Delegates
 	void OnNodeDoubleClicked(class UEdGraphNode* Node);
 	void OnGraphEditorFocused(const TSharedRef<SGraphEditor>& InGraphEditor);
-	FGraphPanelSelectionSet GetSelectedNodes() const;
 
 	void OnAddInputPin();
 	bool CanAddInputPin() const;
@@ -66,21 +65,6 @@ public:
 	bool CanAddBreakpoint() const;
 	void OnRemoveBreakpoint();
 	bool CanRemoveBreakpoint() const;
-	
-	void SelectAllNodes();
-	bool CanSelectAllNodes() const;
-	void DeleteSelectedNodes();
-	bool CanDeleteNodes() const;
-	void DeleteSelectedDuplicatableNodes();
-	void CutSelectedNodes();
-	bool CanCutNodes() const;
-	void CopySelectedNodes();
-	bool CanCopyNodes() const;
-	void PasteNodes();
-	void PasteNodesHere(const FVector2D& Location);
-	bool CanPasteNodes() const;
-	void DuplicateNodes();
-	bool CanDuplicateNodes() const;
 
 	void SearchTree();
 	bool CanSearchTree() const;
@@ -90,7 +74,7 @@ public:
 	bool IsPropertyEditable() const;
 	void OnPackageSaved(const FString& PackageFileName, UObject* Outer);
 	void OnFinishedChangingProperties(const FPropertyChangedEvent& PropertyChangedEvent);
-	void OnClassListUpdated();
+	virtual void OnClassListUpdated() override;
 
 	void UpdateToolbar();
 	bool IsDebuggerReady() const;
@@ -265,7 +249,7 @@ private:
 	void BindDebuggerToolbarCommands();
 
 	/** Called when the selection changes in the GraphEditor */
-	void OnSelectedNodesChanged(const TSet<class UObject*>& NewSelection);
+	virtual void OnSelectedNodesChanged(const TSet<class UObject*>& NewSelection) override;
 
 	/** prepare range of nodes that can be aborted by this decorator */
 	void GetAbortModePreview(const class UBehaviorTreeGraphNode_CompositeDecorator* Node, struct FAbortDrawHelper& Mode0, struct FAbortDrawHelper& Mode1);
@@ -285,15 +269,11 @@ private:
 	/* The Blackboard Data being edited */
 	UBlackboardData* BlackboardData;
 
-	TWeakPtr<SGraphEditor> FocusedGraphEdPtr;
 	TWeakObjectPtr<class UBehaviorTreeGraphNode_CompositeDecorator> FocusedGraphOwner;
 
 	/** Property View */
 	TSharedPtr<class IDetailsView> DetailsView;
 
-	/** The command list for this editor */
-	TSharedPtr<FUICommandList> GraphEditorCommands;
-
 	TSharedPtr<class FBehaviorTreeDebugger> Debugger;
 
 	/** Find results log as well as the search filter */
@@ -305,6 +285,10 @@ private:
 	uint32 bSelectedNodeIsInjected : 1;
 	uint32 SelectedNodesCount;
 
+	uint32 bHasMultipleTaskBP : 1;
+	uint32 bHasMultipleDecoratorBP : 1;
+	uint32 bHasMultipleServiceBP : 1;
+
 	TSharedPtr<class FBehaviorTreeEditorToolbar> ToolbarBuilder;
 
 	/** The details view we use to display the blackboard */
@@ -325,9 +309,6 @@ private:
 	/** Handle to the registered OnPackageSave delegate */
 	FDelegateHandle OnPackageSavedDelegateHandle;
 
-	/** Handle to the registered OnClassListUpdated delegate */
-	FDelegateHandle OnClassListUpdatedDelegateHandle;
-
 public:
 	/** Modes in mode switcher */
 	static const FName BehaviorTreeMode;
diff --git a/Engine/Source/Editor/BehaviorTreeEditor/Private/BehaviorTreeEditorModule.cpp b/Engine/Source/Editor/BehaviorTreeEditor/Private/BehaviorTreeEditorModule.cpp
--- a/Engine/Source/Editor/BehaviorTreeEditor/Private/BehaviorTreeEditorModule.cpp
+++ b/Engine/Source/Editor/BehaviorTreeEditor/Private/BehaviorTreeEditorModule.cpp
@@ -11,6 +11,10 @@
 #include "SGraphNode_Decorator.h"
 #include "EdGraphUtilities.h"
 
+#include "BehaviorTree/Tasks/BTTask_BlueprintBase.h"
+#include "BehaviorTree/Decorators/BTDecorator_BlueprintBase.h"
+#include "BehaviorTree/Services/BTService_BlueprintBase.h"
+
 
 IMPLEMENT_MODULE( FBehaviorTreeEditorModule, BehaviorTreeEditor );
 DEFINE_LOG_CATEGORY(LogBehaviorTreeEditor);
@@ -109,7 +113,11 @@ TSharedRef<IBehaviorTreeEditor> FBehaviorTreeEditorModule::CreateBehaviorTreeEdi
 {
 	if (!ClassCache.IsValid())
 	{
-		ClassCache = MakeShareable(new FClassBrowseHelper);
+		ClassCache = MakeShareable(new FGraphNodeClassHelper(UBTNode::StaticClass()));
+		FGraphNodeClassHelper::AddObservedBlueprintClasses(UBTTask_BlueprintBase::StaticClass());
+		FGraphNodeClassHelper::AddObservedBlueprintClasses(UBTDecorator_BlueprintBase::StaticClass());
+		FGraphNodeClassHelper::AddObservedBlueprintClasses(UBTService_BlueprintBase::StaticClass());
+		ClassCache->UpdateAvailableBlueprintClasses();
 	}
 
 	TSharedRef< FBehaviorTreeEditor > NewBehaviorTreeEditor( new FBehaviorTreeEditor() );
