unified shared parts of AI editors (nested nodes, copy/paste/drag operations)

diff --git a/Engine/Source/Editor/BehaviorTreeEditor/Private/BehaviorTreeEditorTypes.cpp b/Engine/Source/Editor/BehaviorTreeEditor/Private/BehaviorTreeEditorTypes.cpp
--- a/Engine/Source/Editor/BehaviorTreeEditor/Private/BehaviorTreeEditorTypes.cpp
+++ b/Engine/Source/Editor/BehaviorTreeEditor/Private/BehaviorTreeEditorTypes.cpp
@@ -1,554 +1,12 @@
 // Copyright 1998-2015 Epic Games, Inc. All Rights Reserved.
 
 #include "BehaviorTreeEditorPrivatePCH.h"
-#include "Toolkits/AssetEditorManager.h"
-#include "PackageTools.h"
-#include "MessageLog.h"
-#include "AssetRegistryModule.h"
-#include "BehaviorTree/Tasks/BTTask_BlueprintBase.h"
-#include "BehaviorTree/Decorators/BTDecorator_BlueprintBase.h"
-#include "BehaviorTree/Services/BTService_BlueprintBase.h"
-#include "HotReloadInterface.h"
 
 const FString UBehaviorTreeEditorTypes::PinCategory_MultipleNodes("MultipleNodes");
 const FString UBehaviorTreeEditorTypes::PinCategory_SingleComposite("SingleComposite");
 const FString UBehaviorTreeEditorTypes::PinCategory_SingleTask("SingleTask");
 const FString UBehaviorTreeEditorTypes::PinCategory_SingleNode("SingleNode");
 
-#define LOCTEXT_NAMESPACE "SClassViewer"
-
-FClassData::FClassData(UClass* InClass, const FString& InDeprecatedMessage) : 
-	bIsHidden(0), 
-	bHideParent(0), 
-	Class(InClass), 
-	DeprecatedMessage(InDeprecatedMessage)
-{
-	Category = GetCategory();
-}
-
-FClassData::FClassData(const FString& InAssetName, const FString& InGeneratedClassPackage, const FString& InClassName, UClass* InClass) :
-	bIsHidden(0), 
-	bHideParent(0), 
-	Class(InClass), 
-	AssetName(InAssetName), 
-	GeneratedClassPackage(InGeneratedClassPackage), 
-	ClassName(InClassName) 
-{
-	Category = GetCategory();
-}
-
-FString FClassData::ToString() const
-{
-	UClass* MyClass = Class.Get();
-	if (MyClass)
-	{
-		FString ClassDesc = MyClass->GetName();
-
-		if (MyClass->HasAnyClassFlags(CLASS_CompiledFromBlueprint))
-		{
-			return ClassDesc.LeftChop(2);
-		}
-
-		const int32 ShortNameIdx = ClassDesc.Find(TEXT("_"));
-		if (ShortNameIdx != INDEX_NONE)
-		{
-			ClassDesc = ClassDesc.Mid(ShortNameIdx + 1);
-		}
-
-		return ClassDesc;
-	}
-
-	return AssetName;
-}
-
-FString FClassData::GetClassName() const
-{
-	return Class.IsValid() ? Class->GetName() : ClassName;
-}
-
-FString FClassData::GetCategory() const
-{
-	return Class.IsValid() ? Class->GetMetaData(TEXT("Category")) : Category;
-}
-
-bool FClassData::IsAbstract() const
-{
-	return Class.IsValid() ? Class.Get()->HasAnyClassFlags(CLASS_Abstract) : false;
-}
-
-UClass* FClassData::GetClass(bool bSilent)
-{
-	UClass* RetClass = Class.Get();
-	if (RetClass == NULL && GeneratedClassPackage.Len())
-	{
-		GWarn->BeginSlowTask(LOCTEXT("LoadPackage", "Loading Package..."), true);
-
-		UPackage* Package = LoadPackage(NULL, *GeneratedClassPackage, LOAD_NoRedirects);
-		if (Package)
-		{
-			Package->FullyLoad();
-
-			UObject* Object = FindObject<UObject>(Package, *AssetName);
-
-			GWarn->EndSlowTask();
-
-			UBlueprint* BlueprintOb = Cast<UBlueprint>(Object);
-			RetClass = BlueprintOb ? *BlueprintOb->GeneratedClass :
-				Object ? Object->GetClass() : 
-				NULL;
-
-			Class = RetClass;
-		}
-		else
-		{
-			GWarn->EndSlowTask();
-
-			if (!bSilent)
-			{
-				FMessageLog EditorErrors("EditorErrors");
-				EditorErrors.Error(LOCTEXT("PackageLoadFail", "Package Load Failed"));
-				EditorErrors.Info(FText::FromString(GeneratedClassPackage));
-				EditorErrors.Notify(LOCTEXT("PackageLoadFail", "Package Load Failed"));
-			}
-		}
-	}
-
-	return RetClass;
-}
-
-//////////////////////////////////////////////////////////////////////////
-FClassBrowseHelper::FOnPackageListUpdated FClassBrowseHelper::OnPackageListUpdated;
-TArray<FName> FClassBrowseHelper::UnknownPackages;
-bool FClassBrowseHelper::bMoreThanOneTaskClassAvailable = false;
-bool FClassBrowseHelper::bMoreThanOneDecoratorClassAvailable = false;
-bool FClassBrowseHelper::bMoreThanOneServiceClassAvailable = false;
-
-void FClassDataNode::AddUniqueSubNode(TSharedPtr<FClassDataNode> SubNode)
-{
-	for (int32 Idx = 0; Idx < SubNodes.Num(); Idx++)
-	{
-		if (SubNode->Data.GetClassName() == SubNodes[Idx]->Data.GetClassName())
-		{
-			return;
-		}
-	}
-
-	SubNodes.Add(SubNode);
-}
-
-FClassBrowseHelper::FClassBrowseHelper()
-{
-	// Register with the Asset Registry to be informed when it is done loading up files.
-	FAssetRegistryModule& AssetRegistryModule = FModuleManager::GetModuleChecked<FAssetRegistryModule>(TEXT("AssetRegistry"));
-	AssetRegistryModule.Get().OnFilesLoaded().AddRaw( this, &FClassBrowseHelper::InvalidateCache );
-	AssetRegistryModule.Get().OnAssetAdded().AddRaw( this, &FClassBrowseHelper::OnAssetAdded);
-	AssetRegistryModule.Get().OnAssetRemoved().AddRaw( this, &FClassBrowseHelper::OnAssetRemoved );
-
-	// Register to have Populate called when doing a Hot Reload.
-	IHotReloadInterface& HotReloadSupport = FModuleManager::LoadModuleChecked<IHotReloadInterface>("HotReload");
-	HotReloadSupport.OnHotReload().AddRaw(this, &FClassBrowseHelper::OnHotReload);
-
-	// Register to have Populate called when a Blueprint is compiled.
-	GEditor->OnBlueprintCompiled().AddRaw( this, &FClassBrowseHelper::InvalidateCache );
-	GEditor->OnClassPackageLoadedOrUnloaded().AddRaw( this, &FClassBrowseHelper::InvalidateCache );
-
-	UpdateAvailableNodeClasses();
-}
-
-FClassBrowseHelper::~FClassBrowseHelper()
-{
-	// Unregister with the Asset Registry to be informed when it is done loading up files.
-	if (FModuleManager::Get().IsModuleLoaded(TEXT("AssetRegistry")))
-	{
-		FAssetRegistryModule& AssetRegistryModule = FModuleManager::GetModuleChecked<FAssetRegistryModule>(TEXT("AssetRegistry"));
-		AssetRegistryModule.Get().OnFilesLoaded().RemoveAll(this);
-		AssetRegistryModule.Get().OnAssetAdded().RemoveAll(this);
-		AssetRegistryModule.Get().OnAssetRemoved().RemoveAll(this);
-
-		// Unregister to have Populate called when doing a Hot Reload.
-		if (FModuleManager::Get().IsModuleLoaded(TEXT("HotReload")))
-		{
-			IHotReloadInterface& HotReloadSupport = FModuleManager::GetModuleChecked<IHotReloadInterface>("HotReload");
-			HotReloadSupport.OnHotReload().RemoveAll(this);
-		}
-	
-		// Unregister to have Populate called when a Blueprint is compiled.
-		if ( UObjectInitialized() )
-		{
-			// GEditor can't have been destructed before we call this or we'll crash.
-			GEditor->OnBlueprintCompiled().RemoveAll(this);
-			GEditor->OnClassPackageLoadedOrUnloaded().RemoveAll(this);
-		}
-	}
-}
-
-void FClassBrowseHelper::GatherClasses(const UClass* BaseClass, TArray<FClassData>& AvailableClasses)
-{
-	FBehaviorTreeEditorModule& BTEditorModule = FModuleManager::GetModuleChecked<FBehaviorTreeEditorModule>(TEXT("BehaviorTreeEditor"));
-	FClassBrowseHelper* HelperInstance = BTEditorModule.GetClassCache().Get();
-	if (HelperInstance)
-	{
-		const FString BaseClassName = BaseClass->GetName();
-		if (!HelperInstance->RootNode.IsValid())
-		{
-			HelperInstance->BuildClassGraph();
-		}
-
-		TSharedPtr<FClassDataNode> BaseNode = HelperInstance->FindBaseClassNode(HelperInstance->RootNode, BaseClassName);
-		HelperInstance->FindAllSubClasses(BaseNode, AvailableClasses);
-	}
-}
-
-FString FClassBrowseHelper::GetDeprecationMessage(const UClass* Class)
-{
-	static FName MetaDeprecated = TEXT("DeprecatedNode");
-	static FName MetaDeprecatedMessage = TEXT("DeprecationMessage");
-	FString DefDeprecatedMessage("Please remove it!");
-	FString DeprecatedPrefix("DEPRECATED");
-	FString DeprecatedMessage;
-
-	if (Class && Class->HasAnyClassFlags(CLASS_Native) && Class->HasMetaData(MetaDeprecated))
-	{
-		DeprecatedMessage = DeprecatedPrefix + TEXT(": ");
-		DeprecatedMessage += Class->HasMetaData(MetaDeprecatedMessage) ? Class->GetMetaData(MetaDeprecatedMessage) : DefDeprecatedMessage;
-	}
-
-	return DeprecatedMessage;
-}
-
-DEFINE_LOG_CATEGORY_STATIC(LogBTDebug3, Log, All);
-
-bool FClassBrowseHelper::IsClassKnown(const FClassData& ClassData)
-{
-	return !ClassData.IsBlueprint() || !UnknownPackages.Contains(*ClassData.GetPackageName());
-}
-
-void FClassBrowseHelper::AddUnknownClass(const FClassData& ClassData)
-{
-	if (ClassData.IsBlueprint())
-	{
-		UnknownPackages.AddUnique(*ClassData.GetPackageName());
-	}
-}
-
-bool FClassBrowseHelper::IsHidingParentClass(UClass* Class)
-{
-	static FName MetaHideParent = TEXT("HideParentNode");
-	return Class && Class->HasAnyClassFlags(CLASS_Native) && Class->HasMetaData(MetaHideParent);
-}
-
-bool FClassBrowseHelper::IsHidingClass(UClass* Class)
-{
-	static FName MetaHideInEditor = TEXT("HiddenNode");
-	return Class && Class->HasAnyClassFlags(CLASS_Native) && Class->HasMetaData(MetaHideInEditor);
-}
-
-bool FClassBrowseHelper::IsPackageSaved(FName PackageName)
-{
-	const bool bFound = FPackageName::SearchForPackageOnDisk(PackageName.ToString());
-	return bFound;
-}
-
-void FClassBrowseHelper::OnAssetAdded(const class FAssetData& AssetData)
-{
-	TSharedPtr<FClassDataNode> Node = CreateClassDataNode(AssetData);
-	
-	TSharedPtr<FClassDataNode> ParentNode;
-	if (Node.IsValid())
-	{
-		ParentNode = FindBaseClassNode(RootNode, Node->ParentClassName);
-
-		if (!IsPackageSaved(AssetData.PackageName))
-		{
-			UnknownPackages.AddUnique(AssetData.PackageName);
-		}
-		else
-		{
-			const int32 PrevListCount = UnknownPackages.Num();
-			UnknownPackages.RemoveSingleSwap(AssetData.PackageName);
-
-			if (UnknownPackages.Num() != PrevListCount)
-			{
-				FClassBrowseHelper::OnPackageListUpdated.Broadcast();
-			}
-		}
-	}
-
-	if (ParentNode.IsValid())
-	{
-		ParentNode->AddUniqueSubNode(Node);
-		Node->ParentNode = ParentNode;
-	}
-
-	FAssetRegistryModule& AssetRegistryModule = FModuleManager::LoadModuleChecked<FAssetRegistryModule>(TEXT("AssetRegistry"));
-	if(!AssetRegistryModule.Get().IsLoadingAssets())
-	{
-		UpdateAvailableNodeClasses();
-	}
-}
-
-void FClassBrowseHelper::OnAssetRemoved(const class FAssetData& AssetData)
-{
-	const FString* GeneratedClassname = AssetData.TagsAndValues.Find("GeneratedClass");
-	if (GeneratedClassname)
-	{
-		FString AssetClassName = *GeneratedClassname;
-		UObject* Outer1(NULL);
-		ResolveName(Outer1, AssetClassName, false, false);
-		
-		TSharedPtr<FClassDataNode> Node = FindBaseClassNode(RootNode, AssetClassName);
-		if (Node.IsValid() && Node->ParentNode.IsValid())
-		{
-			Node->ParentNode->SubNodes.RemoveSingleSwap(Node);
-		}
-	}
-
-	FAssetRegistryModule& AssetRegistryModule = FModuleManager::LoadModuleChecked<FAssetRegistryModule>(TEXT("AssetRegistry"));
-	if(!AssetRegistryModule.Get().IsLoadingAssets())
-	{
-		UpdateAvailableNodeClasses();
-	}
-}
-
-void FClassBrowseHelper::InvalidateCache()
-{
-	RootNode.Reset();
-
-	UpdateAvailableNodeClasses();
-}
-
-void FClassBrowseHelper::OnHotReload( bool bWasTriggeredAutomatically )
-{
-	InvalidateCache();
-}
-
-TSharedPtr<FClassDataNode> FClassBrowseHelper::CreateClassDataNode(const class FAssetData& AssetData)
-{
-	TSharedPtr<FClassDataNode> Node;
-	const FString* GeneratedClassname = AssetData.TagsAndValues.Find("GeneratedClass");
-	const FString* ParentClassname = AssetData.TagsAndValues.Find("ParentClass");
-
-	if (GeneratedClassname && ParentClassname)
-	{
-		FString AssetClassName = *GeneratedClassname;
-		UObject* Outer1(NULL);
-		ResolveName(Outer1, AssetClassName, false, false);
-
-		FString AssetParentClassName = *ParentClassname;
-		UObject* Outer2(NULL);
-		ResolveName(Outer2, AssetParentClassName, false, false);
-
-		Node = MakeShareable(new FClassDataNode);
-		Node->ParentClassName = AssetParentClassName;
-
-		UObject* AssetOb = AssetData.IsAssetLoaded() ? AssetData.GetAsset() : NULL;
-		UBlueprint* AssetBP = Cast<UBlueprint>(AssetOb);
-		UClass* AssetClass = AssetBP ? *AssetBP->GeneratedClass : AssetOb ? AssetOb->GetClass() : NULL;
-
-		FClassData NewData(AssetData.AssetName.ToString(), AssetData.PackageName.ToString(), AssetClassName, AssetClass);
-		Node->Data = NewData;
-	}
-
-	return Node;
-}
-
-TSharedPtr<FClassDataNode> FClassBrowseHelper::FindBaseClassNode(TSharedPtr<FClassDataNode> Node, const FString& ClassName)
-{
-	TSharedPtr<FClassDataNode> RetNode;
-	if (Node.IsValid())
-	{
-		if (Node->Data.GetClassName() == ClassName)
-		{
-			return Node;
-		}
-
-		for (int32 i = 0; i < Node->SubNodes.Num(); i++)
-		{
-			RetNode = FindBaseClassNode(Node->SubNodes[i], ClassName);
-			if (RetNode.IsValid())
-			{
-				break;
-			}
-		}
-	}
-
-	return RetNode;
-}
-
-void FClassBrowseHelper::FindAllSubClasses(TSharedPtr<FClassDataNode> Node, TArray<FClassData>& AvailableClasses)
-{
-	if (Node.IsValid())
-	{
-		if (!Node->Data.IsAbstract() && !Node->Data.IsDeprecated() && !Node->Data.bIsHidden)
-		{
-			AvailableClasses.Add(Node->Data);
-		}
-
-		for (int32 i = 0; i < Node->SubNodes.Num(); i++)
-		{
-			FindAllSubClasses(Node->SubNodes[i], AvailableClasses);
-		}
-	}
-}
-
-UClass* FClassBrowseHelper::FindAssetClass(const FString& GeneratedClassPackage, const FString& AssetName)
-{
-	UPackage* Package = FindPackage(NULL, *GeneratedClassPackage );
-	if (Package)
-	{
-		UObject* Object = FindObject<UObject>(Package, *AssetName);
-		if (Object)
-		{
-			UBlueprint* BlueprintOb = Cast<UBlueprint>(Object);
-			return BlueprintOb ? *BlueprintOb->GeneratedClass : Object->GetClass();
-		}
-	}
-
-	return NULL;
-}
-
-void FClassBrowseHelper::BuildClassGraph()
-{
-	UClass* RootNodeClass = UBTNode::StaticClass();
-
-	TArray<TSharedPtr<FClassDataNode> > NodeList;
-	TArray<UClass*> HideParentList;
-	RootNode.Reset();
-
-	// gather all native classes
-	for (TObjectIterator<UClass> It; It; ++It)
-	{
-		UClass* TestClass = *It;
-		if (TestClass->HasAnyClassFlags(CLASS_Native) && TestClass->IsChildOf(RootNodeClass))
-		{
-			TSharedPtr<FClassDataNode> NewNode = MakeShareable(new FClassDataNode);
-			NewNode->ParentClassName = TestClass->GetSuperClass()->GetName();
-			
-			FString DeprecatedMessage = GetDeprecationMessage(TestClass);
-			FClassData NewData(TestClass, DeprecatedMessage);
-			
-			NewData.bHideParent = IsHidingParentClass(TestClass);
-			if (NewData.bHideParent)
-			{
-				HideParentList.Add(TestClass->GetSuperClass());
-			}
-
-			NewData.bIsHidden = IsHidingClass(TestClass);
-
-			NewNode->Data = NewData;
-
-			if (TestClass == RootNodeClass)
-			{
-				RootNode = NewNode;
-			}
-
-			NodeList.Add(NewNode);
-		}
-	}
-
-	// find all hidden parent classes
-	for (int32 i = 0; i < NodeList.Num(); i++)
-	{
-		TSharedPtr<FClassDataNode> TestNode = NodeList[i];
-		if (HideParentList.Contains(TestNode->Data.GetClass()))
-		{
-			TestNode->Data.bIsHidden = true;
-		}
-	}
-
-	// gather all blueprints
-	FAssetRegistryModule& AssetRegistryModule = FModuleManager::LoadModuleChecked<FAssetRegistryModule>(TEXT("AssetRegistry"));
-	TArray<FAssetData> BlueprintList;
-
-	FARFilter Filter;
-	Filter.ClassNames.Add(UBlueprint::StaticClass()->GetFName());
-	AssetRegistryModule.Get().GetAssets(Filter, BlueprintList);
-
-	for (int32 i = 0; i < BlueprintList.Num(); i++)
-	{
-		TSharedPtr<FClassDataNode> NewNode = CreateClassDataNode(BlueprintList[i]);
-		NodeList.Add(NewNode);
-	}
-
-	// build class tree
-	AddClassGraphChildren(RootNode, NodeList);
-}
-
-void FClassBrowseHelper::AddClassGraphChildren(TSharedPtr<FClassDataNode> Node, TArray<TSharedPtr<FClassDataNode> >& NodeList)
-{
-	if (!Node.IsValid())
-	{
-		return;
-	}
-
-	const FString NodeClassName = Node->Data.GetClassName();
-	for (int32 i = NodeList.Num() - 1; i >= 0; i--)
-	{
-		if (NodeList[i]->ParentClassName == NodeClassName)
-		{
-			TSharedPtr<FClassDataNode> MatchingNode = NodeList[i];
-			NodeList.RemoveAt(i);
-
-			MatchingNode->ParentNode = Node;
-			Node->SubNodes.Add(MatchingNode);
-
-			AddClassGraphChildren(MatchingNode, NodeList);
-		}
-	}
-}
-
-bool FClassBrowseHelper::IsMoreThanOneTaskClassAvailable()
-{
-	return bMoreThanOneTaskClassAvailable;
-}
-
-bool FClassBrowseHelper::IsMoreThanOneDecoratorClassAvailable()
-{
-	return bMoreThanOneDecoratorClassAvailable;
-}
-
-bool FClassBrowseHelper::IsMoreThanOneServiceClassAvailable()
-{
-	return bMoreThanOneServiceClassAvailable;
-}
-
-void FClassBrowseHelper::UpdateAvailableNodeClasses()
-{
-	if(FModuleManager::Get().IsModuleLoaded(TEXT("AssetRegistry")))
-	{
-		FAssetRegistryModule& AssetRegistryModule = FModuleManager::GetModuleChecked<FAssetRegistryModule>(TEXT("AssetRegistry"));
-
-		const bool bSearchSubClasses = true;
-		{
-			TArray<FName> ClassNames;
-			ClassNames.Add(UBTTask_BlueprintBase::StaticClass()->GetFName());
-			TSet<FName> DerivedClassNames;
-			AssetRegistryModule.Get().GetDerivedClassNames(ClassNames, TSet<FName>(), DerivedClassNames);
-			bMoreThanOneTaskClassAvailable = DerivedClassNames.Num() > 1;
-		}
-
-		{
-			TArray<FName> ClassNames;
-			ClassNames.Add(UBTDecorator_BlueprintBase::StaticClass()->GetFName());
-			TSet<FName> DerivedClassNames;
-			AssetRegistryModule.Get().GetDerivedClassNames(ClassNames, TSet<FName>(), DerivedClassNames);
-			bMoreThanOneDecoratorClassAvailable = DerivedClassNames.Num() > 1;
-		}
-
-		{
-			TArray<FName> ClassNames;
-			ClassNames.Add(UBTService_BlueprintBase::StaticClass()->GetFName());
-			TSet<FName> DerivedClassNames;
-			AssetRegistryModule.Get().GetDerivedClassNames(ClassNames, TSet<FName>(), DerivedClassNames);
-			bMoreThanOneServiceClassAvailable = DerivedClassNames.Num() > 1;
-		}
-	}
-}
-
-//////////////////////////////////////////////////////////////////////////
-
 UBehaviorTreeEditorTypes::UBehaviorTreeEditorTypes(const FObjectInitializer& ObjectInitializer) : Super(ObjectInitializer)
 {
 }
-
-#undef LOCTEXT_NAMESPACE
diff --git a/Engine/Source/Editor/BehaviorTreeEditor/Private/BehaviorTreeGraph.cpp b/Engine/Source/Editor/BehaviorTreeEditor/Private/BehaviorTreeGraph.cpp
--- a/Engine/Source/Editor/BehaviorTreeEditor/Private/BehaviorTreeGraph.cpp
+++ b/Engine/Source/Editor/BehaviorTreeEditor/Private/BehaviorTreeGraph.cpp
@@ -11,10 +11,17 @@
 //////////////////////////////////////////////////////////////////////////
 // BehaviorTreeGraph
 
+namespace BTGraphVersion
+{
+	const int32 Initial = 0;
+	const int32 UnifiedSubNodes = 1;
+
+	const int32 Latest = UnifiedSubNodes;
+}
+
 UBehaviorTreeGraph::UBehaviorTreeGraph(const FObjectInitializer& ObjectInitializer) : Super(ObjectInitializer)
 {
 	Schema = UEdGraphSchema_BehaviorTree::StaticClass();
-	bLockUpdates = false;
 }
 
 void UBehaviorTreeGraph::UpdateBlackboardChange()
@@ -64,7 +71,7 @@ void UBehaviorTreeGraph::UpdateBlackboardChange()
 	}
 }
 
-void UBehaviorTreeGraph::UpdateAsset(EDebuggerFlags DebuggerFlags, bool bBumpVersion)
+void UBehaviorTreeGraph::UpdateAsset(int32 UpdateFlags)
 {
 	if (bLockUpdates)
 	{
@@ -78,7 +85,7 @@ void UBehaviorTreeGraph::UpdateAsset(EDebuggerFlags DebuggerFlags, bool bBumpVer
 		UBehaviorTreeGraphNode* Node = Cast<UBehaviorTreeGraphNode>(Nodes[Index]);
 
 		// debugger flags
-		if (DebuggerFlags == ClearDebuggerFlags)
+		if (UpdateFlags & ClearDebuggerFlags)
 		{
 			Node->ClearDebuggerState();
 
@@ -133,9 +140,9 @@ void UBehaviorTreeGraph::UpdateAsset(EDebuggerFlags DebuggerFlags, bool bBumpVer
 		{
 			CreateBTFromGraph(Node);
 
-			if (bBumpVersion)
+			if ((UpdateFlags & KeepRebuildCounter) == 0)
 			{
-				GraphVersion++;
+				ModCounter++;
 			}
 		}
 	}
@@ -205,14 +212,14 @@ void UBehaviorTreeGraph::UpdateDeprecatedNodes()
 
 			if (Node->NodeInstance)
 			{
-				Node->ErrorMessage = FClassBrowseHelper::GetDeprecationMessage(Node->NodeInstance->GetClass());
+				Node->ErrorMessage = FGraphNodeClassHelper::GetDeprecationMessage(Node->NodeInstance->GetClass());
 			}
 
 			for (int32 i = 0; i < Node->Decorators.Num(); i++)
 			{
 				if (Node->Decorators[i] && Node->Decorators[i]->NodeInstance)
 				{
-					Node->Decorators[i]->ErrorMessage = FClassBrowseHelper::GetDeprecationMessage(Node->Decorators[i]->NodeInstance->GetClass());
+					Node->Decorators[i]->ErrorMessage = FGraphNodeClassHelper::GetDeprecationMessage(Node->Decorators[i]->NodeInstance->GetClass());
 				}
 			}
 
@@ -220,7 +227,7 @@ void UBehaviorTreeGraph::UpdateDeprecatedNodes()
 			{
 				if (Node->Services[i] && Node->Services[i]->NodeInstance)
 				{
-					Node->Services[i]->ErrorMessage = FClassBrowseHelper::GetDeprecationMessage(Node->Services[i]->NodeInstance->GetClass());
+					Node->Services[i]->ErrorMessage = FGraphNodeClassHelper::GetDeprecationMessage(Node->Services[i]->NodeInstance->GetClass());
 				}
 			}
 		}
@@ -817,71 +824,29 @@ void UBehaviorTreeGraph::CreateBTFromGraph(UBehaviorTreeGraphNode* RootEdNode)
 	RemoveOrphanedNodes();
 }
 
-void UBehaviorTreeGraph::RemoveOrphanedNodes()
+void UBehaviorTreeGraph::CollectAllNodeInstances(TSet<UObject*>& NodeInstance)
 {
-	UBehaviorTree* BTAsset = CastChecked<UBehaviorTree>(GetOuter());
+	Super::CollectAllNodeInstances(NodeInstance);
 
-	// Obtain a list of all nodes that should be in the asset
-	TSet<UBTNode*> AllNodes;
-	for (int32 Index = 0; Index < Nodes.Num(); ++Index)
+	for (int32 Idx = 0; Idx < Nodes.Num(); Idx++)
 	{
-		UBehaviorTreeGraphNode* MyNode = Cast<UBehaviorTreeGraphNode>(Nodes[Index]);
-		if (MyNode)
+		UBehaviorTreeGraphNode* MyNode = Cast<UBehaviorTreeGraphNode>(Nodes[Idx]);
+		for (int32 SubIdx = 0; SubIdx < MyNode->Decorators.Num(); SubIdx++)
 		{
-			UBTNode* MyNodeInstance = Cast<UBTNode>(MyNode->NodeInstance);
-			if (MyNodeInstance)
+			UBehaviorTreeGraphNode_CompositeDecorator* SubgraphNode = Cast<UBehaviorTreeGraphNode_CompositeDecorator>(MyNode->Decorators[SubIdx]);
+			if (SubgraphNode)
 			{
-				AllNodes.Add(MyNodeInstance);
-			}
-
-			for (int32 iDecorator = 0; iDecorator < MyNode->Decorators.Num(); iDecorator++)
-			{
-				UBehaviorTreeGraphNode_CompositeDecorator* SubgraphNode = Cast<UBehaviorTreeGraphNode_CompositeDecorator>(MyNode->Decorators[iDecorator]);
-				if (SubgraphNode)
-				{
-					TArray<UBTDecorator*> NodeInstances;
-					TArray<FBTDecoratorLogic> DummyOps;
-					SubgraphNode->CollectDecoratorData(NodeInstances, DummyOps);
-
-					for (int32 SubIdx = 0; SubIdx < NodeInstances.Num(); SubIdx++)
-					{
-						AllNodes.Add(NodeInstances[SubIdx]);
-					}
-				}
-				else
-				{
-					UBTNode* MyDecoratorNodeInstance = MyNode->Decorators[iDecorator] ? Cast<UBTNode>(MyNode->Decorators[iDecorator]->NodeInstance) : NULL;
-					if (MyDecoratorNodeInstance)
-					{
-						AllNodes.Add(MyDecoratorNodeInstance);
-					}
-				}
-			}
+				TArray<UBTDecorator*> DecoratorInstances;
+				TArray<FBTDecoratorLogic> DummyOps;
+				SubgraphNode->CollectDecoratorData(DecoratorInstances, DummyOps);
 
-			for (int32 iService = 0; iService < MyNode->Services.Num(); iService++)
-			{
-				UBTNode* MyServiceNodeInstance = MyNode->Services[iService] ? Cast<UBTNode>(MyNode->Services[iService]->NodeInstance) : NULL;
-				if (MyServiceNodeInstance)
+				for (int32 DecoratorIdx = 0; DecoratorIdx < DecoratorInstances.Num(); DecoratorIdx++)
 				{
-					AllNodes.Add(MyServiceNodeInstance);
+					NodeInstance.Add(DecoratorInstances[DecoratorIdx]);
 				}
 			}
 		}
 	}
-
-	// Obtain a list of all nodes actually in the asset and discard unused nodes
-	TArray<UObject*> AllInners;
-	const bool bIncludeNestedObjects = false;
-	GetObjectsWithOuter(BTAsset, AllInners, bIncludeNestedObjects);
-	for (auto InnerIt = AllInners.CreateConstIterator(); InnerIt; ++InnerIt)
-	{
-		UBTNode* Node = Cast<UBTNode>(*InnerIt);
-		if (Node && !AllNodes.Contains(Node))
-		{
-			Node->SetFlags(RF_Transient);
-			Node->Rename(NULL, GetTransientPackage(), REN_DontCreateRedirectors | REN_NonTransactional | REN_ForceNoResetLoaders);
-		}
-	}
 }
 
 void UBehaviorTreeGraph::SpawnMissingNodes()
@@ -929,37 +894,6 @@ void UBehaviorTreeGraph::UpdateAbortHighlight(struct FAbortDrawHelper& Mode0, st
 	}
 }
 
-bool UBehaviorTreeGraph::UpdateUnknownNodeClasses()
-{
-	bool bUpdated = false;
-	for (int32 NodeIdx = 0; NodeIdx < Nodes.Num(); NodeIdx++)
-	{
-		UBehaviorTreeGraphNode* MyNode = Cast<UBehaviorTreeGraphNode>(Nodes[NodeIdx]);
-		const bool bUpdatedNode = MyNode->RefreshNodeClass();
-		bUpdated = bUpdated || bUpdatedNode;
-
-		for (int32 SubNodeIdx = 0; SubNodeIdx < MyNode->Decorators.Num(); SubNodeIdx++)
-		{
-			if (MyNode->Decorators[SubNodeIdx])
-			{
-				const bool bUpdatedSubNode = MyNode->Decorators[SubNodeIdx]->RefreshNodeClass();
-				bUpdated = bUpdated || bUpdatedSubNode;
-			}
-		}
-
-		for (int32 SubNodeIdx = 0; SubNodeIdx < MyNode->Services.Num(); SubNodeIdx++)
-		{
-			if (MyNode->Services[SubNodeIdx])
-			{
-				const bool bUpdatedSubNode = MyNode->Services[SubNodeIdx]->RefreshNodeClass();
-				bUpdated = bUpdated || bUpdatedSubNode;
-			}
-		}
-	}
-
-	return bUpdated;
-}
-
 bool UBehaviorTreeGraph::UpdateInjectedNodes()
 {
 	bool bHasUpdated = false;
@@ -1039,18 +973,6 @@ void UBehaviorTreeGraph::RebuildExecutionOrder()
 	}
 }
 
-void UBehaviorTreeGraph::LockUpdates()
-{
-	bLockUpdates = true;
-}
-
-void UBehaviorTreeGraph::UnlockUpdates()
-{
-	bLockUpdates = false;
-
-	UpdateAsset(EDebuggerFlags::SkipDebuggerFlags);
-}
-
 namespace BTAutoArrangeHelpers
 {
 	struct FNodeBoundsInfo
@@ -1141,3 +1063,60 @@ void UBehaviorTreeGraph::AutoArrange()
 
 	RootNode->NodeWidget.Pin()->GetOwnerPanel()->ZoomToFit(/*bOnlySelection=*/ false);
 }
+
+void UBehaviorTreeGraph::OnSubNodeDropped()
+{
+	Super::OnSubNodeDropped();
+
+	FAbortDrawHelper EmptyMode;
+	UpdateAsset(UBehaviorTreeGraph::ClearDebuggerFlags);
+	UpdateAbortHighlight(EmptyMode, EmptyMode);
+}
+
+void UBehaviorTreeGraph::UpdateVersion()
+{
+	if (!bIsUsingModCounter)
+	{
+		bIsUsingModCounter = true;
+		GraphVersion = BTGraphVersion::Initial;
+	}
+
+	if (GraphVersion == BTGraphVersion::Latest)
+	{
+		return;
+	}
+
+	// convert to nested nodes
+	if (GraphVersion < BTGraphVersion::UnifiedSubNodes)
+	{
+		UpdateVersion_UnifiedSubNodes();
+	}
+
+	GraphVersion = BTGraphVersion::Latest;
+	Modify();
+}
+
+void UBehaviorTreeGraph::MarkVersion()
+{
+	GraphVersion = BTGraphVersion::Latest;
+	bIsUsingModCounter = true;
+}
+
+void UBehaviorTreeGraph::UpdateVersion_UnifiedSubNodes()
+{
+	for (int32 NodeIdx = 0; NodeIdx < Nodes.Num(); NodeIdx++)
+	{
+		UBehaviorTreeGraphNode* MyNode = Cast<UBehaviorTreeGraphNode>(Nodes[NodeIdx]);
+		MyNode->SubNodes.Reset(MyNode->Decorators.Num() + MyNode->Services.Num());
+
+		for (int32 SubIdx = 0; SubIdx < MyNode->Decorators.Num(); SubIdx++)
+		{
+			MyNode->SubNodes.Add(MyNode->Decorators[SubIdx]);
+		}
+
+		for (int32 SubIdx = 0; SubIdx < MyNode->Services.Num(); SubIdx++)
+		{
+			MyNode->SubNodes.Add(MyNode->Services[SubIdx]);
+		}
+	}
+}
diff --git a/Engine/Source/Editor/BehaviorTreeEditor/Private/BehaviorTreeGraphNode.cpp b/Engine/Source/Editor/BehaviorTreeEditor/Private/BehaviorTreeGraphNode.cpp
--- a/Engine/Source/Editor/BehaviorTreeEditor/Private/BehaviorTreeGraphNode.cpp
+++ b/Engine/Source/Editor/BehaviorTreeEditor/Private/BehaviorTreeGraphNode.cpp
@@ -2,7 +2,7 @@
 
 #include "BehaviorTreeEditorPrivatePCH.h"
 #include "ScopedTransaction.h"
-#include "SGraphEditorActionMenu_BehaviorTree.h"
+#include "SGraphEditorActionMenuAI.h"
 #include "BlueprintNodeHelpers.h"
 #include "BehaviorTree/BehaviorTree.h"
 #include "AssetData.h"
@@ -10,11 +10,10 @@
 #include "BehaviorTree/Services/BTService_BlueprintBase.h"
 #include "BehaviorTree/Tasks/BTTask_BlueprintBase.h"
 
-#define LOCTEXT_NAMESPACE "BehaviorTreeGraphNode"
+#define LOCTEXT_NAMESPACE "BehaviorTreeEditor"
 
 UBehaviorTreeGraphNode::UBehaviorTreeGraphNode(const FObjectInitializer& ObjectInitializer) : Super(ObjectInitializer)
 {
-	NodeInstance = NULL;
 	bHighlightInAbortRange0 = false;
 	bHighlightInAbortRange1 = false;
 	bHighlightInSearchRange0 = false;
@@ -44,161 +43,52 @@ void UBehaviorTreeGraphNode::AllocateDefaultPins()
 	CreatePin(EGPD_Output, UBehaviorTreeEditorTypes::PinCategory_MultipleNodes, TEXT(""), NULL, false, false, TEXT("Out"));
 }
 
-void UBehaviorTreeGraphNode::PostPlacedNewNode()
+void UBehaviorTreeGraphNode::InitializeInstance()
 {
-	UClass* NodeClass = ClassData.GetClass(true);
-	if (NodeClass)
+	UBTNode* BTNode = Cast<UBTNode>(NodeInstance);
+	UBehaviorTree* BTAsset = BTNode ? Cast<UBehaviorTree>(BTNode->GetOuter()) : nullptr;
+	if (BTNode && BTAsset)
 	{
-		UBehaviorTree* BT = Cast<UBehaviorTree>(GetBehaviorTreeGraph()->GetOuter());
-		if (BT)
-		{
-			NodeInstance = NewObject<UBTNode>(BT, NodeClass);
-
-			UBTNode* BTNode = (UBTNode*)NodeInstance;
-			BTNode->SetFlags(RF_Transactional);
-			BTNode->InitializeFromAsset(*BT);
-			BTNode->InitializeNode(NULL, MAX_uint16, 0, 0);
-		}
-	}
-}
-
-bool UBehaviorTreeGraphNode::CanDuplicateNode() const
-{
-	return bInjectedNode ? false : Super::CanDuplicateNode();
-}
-
-bool UBehaviorTreeGraphNode::CanUserDeleteNode() const
-{
-	return bInjectedNode ? false : Super::CanUserDeleteNode();
-}
-
-void UBehaviorTreeGraphNode::PrepareForCopying()
-{
-	if (NodeInstance)
-	{
-		// Temporarily take ownership of the node instance, so that it is not deleted when cutting
-		NodeInstance->Rename(NULL, this, REN_DontCreateRedirectors | REN_DoNotDirty );
-	}
-}
-#if WITH_EDITOR
-
-void UBehaviorTreeGraphNode::PostEditImport()
-{
-	ResetNodeOwner();
-
-	if (NodeInstance)
-	{
-		UBehaviorTree* BT = Cast<UBehaviorTree>(GetBehaviorTreeGraph()->GetOuter());
-		if (BT)
-		{
-			UBTNode* BTNode = (UBTNode*)NodeInstance;
-			BTNode->InitializeFromAsset(*BT);
-			BTNode->InitializeNode(NULL, MAX_uint16, 0, 0);
-		}
-	}
-}
-
-void UBehaviorTreeGraphNode::PostEditUndo()
-{
-	ResetNodeOwner();
-}
-
-#endif
-
-void UBehaviorTreeGraphNode::PostCopyNode()
-{
-	ResetNodeOwner();
-}
-
-void UBehaviorTreeGraphNode::ResetNodeOwner()
-{
-	if (NodeInstance)
-	{
-		UBehaviorTree* BT = Cast<UBehaviorTree>(GetBehaviorTreeGraph()->GetOuter());
-		NodeInstance->Rename(NULL, BT, REN_DontCreateRedirectors | REN_DoNotDirty);
-		NodeInstance->ClearFlags(RF_Transient);
-
-		// also reset decorators and services
-		for(auto& Decorator : Decorators)
-		{
-			if(Decorator->NodeInstance != nullptr)
-			{
-				Decorator->NodeInstance->Rename(NULL, BT, REN_DontCreateRedirectors | REN_DoNotDirty);
-				Decorator->NodeInstance->ClearFlags(RF_Transient);
-			}
-		}
-
-		for(auto& Service : Services)
-		{
-			if(Service->NodeInstance != nullptr)
-			{
-				Service->NodeInstance->Rename(NULL, BT, REN_DontCreateRedirectors | REN_DoNotDirty);
-				Service->NodeInstance->ClearFlags(RF_Transient);
-			}
-		}
+		BTNode->InitializeFromAsset(*BTAsset);
+		BTNode->InitializeNode(NULL, MAX_uint16, 0, 0);
 	}
 }
 
-FString	UBehaviorTreeGraphNode::GetDescription() const
+FText UBehaviorTreeGraphNode::GetDescription() const
 {
-	if (const UBTNode* Node = Cast<UBTNode>(NodeInstance))
+	const UBTNode* BTNode = Cast<const UBTNode>(NodeInstance);
+	if (BTNode)
 	{
-		return Node->GetStaticDescription();
+		return FText::FromString(BTNode->GetStaticDescription());
 	}
 
-	FString StoredClassName = ClassData.GetClassName();
-	StoredClassName.RemoveFromEnd(TEXT("_C"));
-
-	return FString::Printf(TEXT("%s: %s"), *StoredClassName,
-		*LOCTEXT("NodeClassError", "Class not found, make sure it's saved!").ToString());
+	return Super::GetDescription();
 }
 
 FText UBehaviorTreeGraphNode::GetTooltipText() const
 {
 	FText TooltipDesc;
 
-	if(!NodeInstance)
-	{
-		TooltipDesc = LOCTEXT("NodeClassError", "Class not found, make sure it's saved!");
-	}
-	else
+	if (NodeInstance)
 	{
-		if(bHasObserverError)
+		if (bHasObserverError)
 		{
 			TooltipDesc = LOCTEXT("ObserverError", "Observer has invalid abort setting!");
 		}
-		else if(DebuggerRuntimeDescription.Len() > 0)
+		else if (DebuggerRuntimeDescription.Len() > 0)
 		{
 			TooltipDesc = FText::FromString(DebuggerRuntimeDescription);
 		}
-		else if(ErrorMessage.Len() > 0)
-		{
-			TooltipDesc = FText::FromString(ErrorMessage);
-		}
-		else
-		{
-			if( NodeInstance->GetClass()->IsChildOf(UBTDecorator_BlueprintBase::StaticClass()) ||
-				NodeInstance->GetClass()->IsChildOf(UBTService_BlueprintBase::StaticClass()) ||
-				NodeInstance->GetClass()->IsChildOf(UBTTask_BlueprintBase::StaticClass()) )
-			{
-				FAssetData AssetData(NodeInstance->GetClass()->ClassGeneratedBy);
-				const FString* Description = AssetData.TagsAndValues.Find(GET_MEMBER_NAME_CHECKED(UBlueprint, BlueprintDescription));
+	}
 
-				if (Description && !Description->IsEmpty())
-				{
-					TooltipDesc = FText::FromString(Description->Replace(TEXT("\\n"), TEXT("\n")));
-				}
-			}
-			else
-			{
-				TooltipDesc = NodeInstance->GetClass()->GetToolTipText();
-			}	
-		}
+	if (TooltipDesc.IsEmpty())
+	{
+		TooltipDesc = Super::GetTooltipText();
 	}
 
 	if (bInjectedNode)
 	{
-		FText const InjectedDesc = !TooltipDesc.IsEmpty() ? TooltipDesc : FText::FromString(GetDescription());
+		FText const InjectedDesc = !TooltipDesc.IsEmpty() ? TooltipDesc : GetDescription();
 		// @TODO: FText::Format() is slow... consider caching this tooltip like 
 		//        we do for a lot of the BP nodes now (unfamiliar with this 
 		//        node's purpose, so hesitant to muck with this at this time).
@@ -208,212 +98,133 @@ FText UBehaviorTreeGraphNode::GetTooltipText() const
 	return TooltipDesc;
 }
 
-UEdGraphPin* UBehaviorTreeGraphNode::GetInputPin(int32 InputIndex) const
+UBehaviorTreeGraph* UBehaviorTreeGraphNode::GetBehaviorTreeGraph()
 {
-	check(InputIndex >= 0);
-
-	for (int32 PinIndex = 0, FoundInputs = 0; PinIndex < Pins.Num(); PinIndex++)
-	{
-		if (Pins[PinIndex]->Direction == EGPD_Input)
-		{
-			if (InputIndex == FoundInputs)
-			{
-				return Pins[PinIndex];
-			}
-			else
-			{
-				FoundInputs++;
-			}
-		}
-	}
+	return CastChecked<UBehaviorTreeGraph>(GetGraph());
+}
 
-	return NULL;
+bool UBehaviorTreeGraphNode::CanCreateUnderSpecifiedSchema(const UEdGraphSchema* DesiredSchema) const
+{
+	return DesiredSchema->GetClass()->IsChildOf(UEdGraphSchema_BehaviorTree::StaticClass());
 }
 
-UEdGraphPin* UBehaviorTreeGraphNode::GetOutputPin(int32 InputIndex) const
+void UBehaviorTreeGraphNode::OnSubNodeAdded(UAIGraphNode* NodeTemplate)
 {
-	check(InputIndex >= 0);
+	UBehaviorTreeGraphNode* BTGraphNode = Cast<UBehaviorTreeGraphNode>(NodeTemplate);
 
-	for (int32 PinIndex = 0, FoundInputs = 0; PinIndex < Pins.Num(); PinIndex++)
+	if (Cast<UBehaviorTreeGraphNode_CompositeDecorator>(NodeTemplate) || Cast<UBehaviorTreeGraphNode_Decorator>(NodeTemplate))
 	{
-		if (Pins[PinIndex]->Direction == EGPD_Output)
+		bool bAppend = true;
+		for (int32 Idx = 0; Idx < Decorators.Num(); Idx++)
 		{
-			if (InputIndex == FoundInputs)
-			{
-				return Pins[PinIndex];
-			}
-			else
+			if (Decorators[Idx]->bInjectedNode)
 			{
-				FoundInputs++;
+				Decorators.Insert(BTGraphNode, Idx);
+				bAppend = false;
+				break;
 			}
 		}
-	}
 
-	return NULL;
+		if (bAppend)
+		{
+			Decorators.Add(BTGraphNode);
+		}
+	} 
+	else
+	{
+		Services.Add(BTGraphNode);
+	}
 }
 
-void UBehaviorTreeGraphNode::AutowireNewNode(UEdGraphPin* FromPin)
+void UBehaviorTreeGraphNode::OnSubNodeRemoved(UAIGraphNode* SubNode)
 {
-	Super::AutowireNewNode(FromPin);
+	const int32 DecoratorIdx = Decorators.IndexOfByKey(SubNode);
+	const int32 ServiceIdx = Services.IndexOfByKey(SubNode);
 
-	if (FromPin != nullptr)
+	if (DecoratorIdx >= 0)
 	{
-		UEdGraphPin* OutputPin = GetOutputPin();
+		Decorators.RemoveAt(DecoratorIdx);
+	}
 
-		if (GetSchema()->TryCreateConnection(FromPin, GetInputPin()))
-		{
-			FromPin->GetOwningNode()->NodeConnectionListChanged();
-		}
-		else if(OutputPin != nullptr && GetSchema()->TryCreateConnection(OutputPin, FromPin))
-		{
-			NodeConnectionListChanged();
-		}
+	if (ServiceIdx >= 0)
+	{
+		Services.RemoveAt(ServiceIdx);
 	}
 }
 
-
-UBehaviorTreeGraph* UBehaviorTreeGraphNode::GetBehaviorTreeGraph()
+void UBehaviorTreeGraphNode::RemoveAllSubNodes()
 {
-	return CastChecked<UBehaviorTreeGraph>(GetGraph());
-}
+	Super::RemoveAllSubNodes();
 
-void UBehaviorTreeGraphNode::NodeConnectionListChanged()
-{
-	GetBehaviorTreeGraph()->UpdateAsset(UBehaviorTreeGraph::SkipDebuggerFlags);
+	Decorators.Reset();
+	Services.Reset();
 }
 
-
-bool UBehaviorTreeGraphNode::CanCreateUnderSpecifiedSchema(const UEdGraphSchema* DesiredSchema) const
+int32 UBehaviorTreeGraphNode::FindSubNodeDropIndex(UAIGraphNode* SubNode) const
 {
-	return DesiredSchema->GetClass()->IsChildOf(UEdGraphSchema_BehaviorTree::StaticClass());
+	const int32 SubIdx = SubNodes.IndexOfByKey(SubNode) + 1;
+	const int32 DecoratorIdx = Decorators.IndexOfByKey(SubNode) + 1;
+	const int32 ServiceIdx = Services.IndexOfByKey(SubNode) + 1;
+
+	const int32 CombinedIdx = (SubIdx & 0xff) | ((DecoratorIdx & 0xff) << 8) | ((ServiceIdx & 0xff) << 16);
+	return CombinedIdx;
 }
 
-void UBehaviorTreeGraphNode::DiffProperties(UStruct* Struct, void* DataA, void* DataB, FDiffResults& Results, FDiffSingleResult& Diff)
+void UBehaviorTreeGraphNode::InsertSubNodeAt(UAIGraphNode* SubNode, int32 DropIndex)
 {
-	for (TFieldIterator<UProperty> PropertyIt(Struct, EFieldIteratorFlags::IncludeSuper); PropertyIt; ++PropertyIt)
-	{
-		UProperty* Prop = *PropertyIt;
- 		// skip properties we cant see
-		if (!Prop->HasAnyPropertyFlags(CPF_Edit|CPF_BlueprintVisible) ||
-			Prop->HasAnyPropertyFlags(CPF_Transient) ||
-			Prop->HasAnyPropertyFlags(CPF_DisableEditOnInstance) ||
-			Prop->IsA(UFunction::StaticClass()) ||
-			Prop->IsA(UDelegateProperty::StaticClass()) ||
-			Prop->IsA(UMulticastDelegateProperty::StaticClass()))
-		{
-			continue;
-		}
-		
-		FString ValueStringA = BlueprintNodeHelpers::DescribeProperty(Prop, Prop->ContainerPtrToValuePtr<uint8>(DataA));
-		FString ValueStringB  = BlueprintNodeHelpers::DescribeProperty(Prop, Prop->ContainerPtrToValuePtr<uint8>(DataB));
+	const int32 SubIdx = (DropIndex & 0xff) - 1;
+	const int32 DecoratorIdx = ((DropIndex >> 8) & 0xff) - 1;
+	const int32 ServiceIdx = ((DropIndex >> 16) & 0xff) - 1;
 
-		if ( ValueStringA != ValueStringB )
-		{
-			if(Results)
-			{
-				Diff.DisplayString = FText::Format(LOCTEXT("DIF_NodePropertyFmt", "Property Changed: {0} "), FText::FromString(Prop->GetName()));
-				Results.Add(Diff);
-			}
-		}
+	if (SubIdx >= 0)
+	{
+		SubNodes.Insert(SubNode, SubIdx);
 	}
-}
-
-void UBehaviorTreeGraphNode::FindDiffs(UEdGraphNode* OtherNode, FDiffResults& Results)
-{
-	FDiffSingleResult Diff;
-	Diff.Diff = EDiffType::NODE_PROPERTY;
-	Diff.Node1 = this;
-	Diff.Node2 = OtherNode;
-	Diff.ToolTip = LOCTEXT("DIF_NodePropertyToolTip", "A Property of the node has changed");
-	Diff.DisplayColor = FLinearColor(0.25f,0.71f,0.85f);
-	
-	UBehaviorTreeGraphNode* ThisBehaviorTreeNode = Cast<UBehaviorTreeGraphNode>(this);
-	UBehaviorTreeGraphNode* OtherBehaviorTreeNode = Cast<UBehaviorTreeGraphNode>(OtherNode);
-	if(ThisBehaviorTreeNode && OtherBehaviorTreeNode)
+	else
 	{
-		DiffProperties( ThisBehaviorTreeNode->GetClass(), ThisBehaviorTreeNode, OtherBehaviorTreeNode, Results, Diff );
-
-		UBTNode* ThisNodeInstance = Cast<UBTNode>(ThisBehaviorTreeNode->NodeInstance);
-	    UBTNode* OtherNodeInstance = Cast<UBTNode>(OtherBehaviorTreeNode->NodeInstance);
-	    if(ThisNodeInstance && OtherNodeInstance)
-	    {
-		    DiffProperties( ThisNodeInstance->GetClass(), ThisNodeInstance, OtherNodeInstance, Results, Diff );
-		}
+		SubNodes.Add(SubNode);
 	}
-}
-
-void UBehaviorTreeGraphNode::AddSubNode(UBehaviorTreeGraphNode* NodeTemplate, class UEdGraph* ParentGraph)
-{
-	const FScopedTransaction Transaction(LOCTEXT("AddNode", "Add Node"));
-	ParentGraph->Modify();
-	Modify();
-
-	NodeTemplate->SetFlags(RF_Transactional);
-
-	// set outer to be the graph so it doesn't go away
-	NodeTemplate->Rename(NULL, ParentGraph, REN_NonTransactional);
-	NodeTemplate->ParentNode = this;
-
-	NodeTemplate->CreateNewGuid();
-	NodeTemplate->PostPlacedNewNode();
-	NodeTemplate->AllocateDefaultPins();
-	NodeTemplate->AutowireNewNode(NULL);
 
-	NodeTemplate->NodePosX = 0;
-	NodeTemplate->NodePosY = 0;
+	UBehaviorTreeGraphNode* TypedNode = Cast<UBehaviorTreeGraphNode>(SubNode);
+	const bool bIsDecorator = Cast<UBTDecorator>(SubNode->NodeInstance) != nullptr;
+	const bool bIsService = Cast<UBTService>(SubNode->NodeInstance) != nullptr;
 
-	if (Cast<UBehaviorTreeGraphNode_CompositeDecorator>(NodeTemplate) || Cast<UBehaviorTreeGraphNode_Decorator>(NodeTemplate))
+	if (TypedNode)
 	{
-		bool bAppend = true;
-		for (int32 Idx = 0; Idx < Decorators.Num(); Idx++)
+		if (bIsDecorator)
 		{
-			if (Decorators[Idx]->bInjectedNode)
+			if (DecoratorIdx >= 0)
 			{
-				Decorators.Insert(NodeTemplate, Idx);
-				bAppend = false;
-				break;
+				Decorators.Insert(TypedNode, DecoratorIdx);
+			}
+			else
+			{
+				Decorators.Add(TypedNode);
 			}
+
 		}
 
-		if (bAppend)
+		if (bIsService)
 		{
-			Decorators.Add(NodeTemplate);
+			if (ServiceIdx >= 0)
+			{
+				Services.Insert(TypedNode, ServiceIdx);
+			}
+			else
+			{
+				Services.Add(TypedNode);
+			}
 		}
-	} 
-	else
-	{
-		Services.Add(NodeTemplate);
 	}
-
-	ParentGraph->NotifyGraphChanged();
-	GetBehaviorTreeGraph()->UpdateAsset(UBehaviorTreeGraph::SkipDebuggerFlags);
-}
-
-
-FString GetShortTypeNameHelper(const UObject* Ob)
-{
-	if (Ob->GetClass()->HasAnyClassFlags(CLASS_CompiledFromBlueprint))
-	{
-		return Ob->GetClass()->GetName().LeftChop(2);
-	}
-
-	FString TypeDesc = Ob->GetClass()->GetName();
-	const int32 ShortNameIdx = TypeDesc.Find(TEXT("_"));
-	if (ShortNameIdx != INDEX_NONE)
-	{
-		TypeDesc = TypeDesc.Mid(ShortNameIdx + 1);
-	}
-
-	return TypeDesc;
 }
 
 void UBehaviorTreeGraphNode::CreateAddDecoratorSubMenu(class FMenuBuilder& MenuBuilder, UEdGraph* Graph) const
 {
-	TSharedRef<SGraphEditorActionMenu_BehaviorTree> Menu =	
-		SNew(SGraphEditorActionMenu_BehaviorTree)
+	TSharedRef<SGraphEditorActionMenuAI> Menu =	
+		SNew(SGraphEditorActionMenuAI)
 		.GraphObj( Graph )
 		.GraphNode((UBehaviorTreeGraphNode*)this)
-		.SubNodeType(ESubNode::Decorator)
+		.SubNodeFlags(ESubNode::Decorator)
 		.AutoExpandActionMenu(true);
 
 	MenuBuilder.AddWidget(Menu,FText(),true);
@@ -421,11 +232,11 @@ void UBehaviorTreeGraphNode::CreateAddDecoratorSubMenu(class FMenuBuilder& MenuB
 
 void UBehaviorTreeGraphNode::CreateAddServiceSubMenu(class FMenuBuilder& MenuBuilder, UEdGraph* Graph) const
 {
-	TSharedRef<SGraphEditorActionMenu_BehaviorTree> Menu =	
-		SNew(SGraphEditorActionMenu_BehaviorTree)
+	TSharedRef<SGraphEditorActionMenuAI> Menu =	
+		SNew(SGraphEditorActionMenuAI)
 		.GraphObj( Graph )
 		.GraphNode((UBehaviorTreeGraphNode*)this)
-		.SubNodeType(ESubNode::Service)
+		.SubNodeFlags(ESubNode::Service)
 		.AutoExpandActionMenu(true);
 
 	MenuBuilder.AddWidget(Menu,FText(),true);
@@ -447,17 +258,6 @@ void UBehaviorTreeGraphNode::AddContextMenuActionsServices(const FGraphNodeConte
 		FNewMenuDelegate::CreateUObject( this, &UBehaviorTreeGraphNode::CreateAddServiceSubMenu,(UEdGraph*)Context.Graph));
 }
 
-void UBehaviorTreeGraphNode::DestroyNode()
-{
-	if (ParentNode)
-	{
-		ParentNode->Modify();
-		ParentNode->Decorators.Remove(this);
-		ParentNode->Services.Remove(this);
-	}
-	UEdGraphNode::DestroyNode();
-}
-
 void UBehaviorTreeGraphNode::ClearDebuggerState()
 {
 	bHasBreakpoint = false;
@@ -481,35 +281,9 @@ FName UBehaviorTreeGraphNode::GetNameIcon() const
 	return BTNodeInstance != nullptr ? BTNodeInstance->GetNodeIconName() : FName("BTEditor.Graph.BTNode.Icon");
 }
 
-bool UBehaviorTreeGraphNode::UsesBlueprint() const
-{
-	UBTNode* BTNodeInstance = Cast<UBTNode>(NodeInstance);
-	return BTNodeInstance != nullptr ? BTNodeInstance->UsesBlueprint() : false;
-}
-
-bool UBehaviorTreeGraphNode::RefreshNodeClass()
-{
-	bool bUpdated = false;
-	if (NodeInstance == NULL)
-	{
-		if (FClassBrowseHelper::IsClassKnown(ClassData))
-		{
-			PostPlacedNewNode();
-			bUpdated = (NodeInstance != NULL);
-		}
-		else
-		{
-			FClassBrowseHelper::AddUnknownClass(ClassData);
-		}
-	}
-
-	return bUpdated;
-}
-
 bool UBehaviorTreeGraphNode::HasErrors() const
 {
-	return bHasObserverError || ErrorMessage.Len() > 0 || NodeInstance == NULL;
+	return bHasObserverError || Super::HasErrors();
 }
 
-
 #undef LOCTEXT_NAMESPACE
diff --git a/Engine/Source/Editor/BehaviorTreeEditor/Private/BehaviorTreeGraphNode_Composite.cpp b/Engine/Source/Editor/BehaviorTreeEditor/Private/BehaviorTreeGraphNode_Composite.cpp
--- a/Engine/Source/Editor/BehaviorTreeEditor/Private/BehaviorTreeGraphNode_Composite.cpp
+++ b/Engine/Source/Editor/BehaviorTreeEditor/Private/BehaviorTreeGraphNode_Composite.cpp
@@ -2,8 +2,6 @@
 
 #include "BehaviorTreeEditorPrivatePCH.h"
 
-#define LOCTEXT_NAMESPACE "BehaviorTreeGraphNode_Composite"
-
 UBehaviorTreeGraphNode_Composite::UBehaviorTreeGraphNode_Composite(const FObjectInitializer& ObjectInitializer)
 	: Super(ObjectInitializer)
 {
@@ -24,5 +22,3 @@ void UBehaviorTreeGraphNode_Composite::GetContextMenuActions(const FGraphNodeCon
 	AddContextMenuActionsDecorators(Context);
 	AddContextMenuActionsServices(Context);
 }
-
-#undef LOCTEXT_NAMESPACE
diff --git a/Engine/Source/Editor/BehaviorTreeEditor/Private/BehaviorTreeGraphNode_CompositeDecorator.cpp b/Engine/Source/Editor/BehaviorTreeEditor/Private/BehaviorTreeGraphNode_CompositeDecorator.cpp
--- a/Engine/Source/Editor/BehaviorTreeEditor/Private/BehaviorTreeGraphNode_CompositeDecorator.cpp
+++ b/Engine/Source/Editor/BehaviorTreeEditor/Private/BehaviorTreeGraphNode_CompositeDecorator.cpp
@@ -5,7 +5,7 @@
 #include "BehaviorTree/BehaviorTree.h"
 #include "BehaviorTreeGraphNode_CompositeDecorator.h"
 
-#define LOCTEXT_NAMESPACE "BehaviorTreeGraphNode"
+#define LOCTEXT_NAMESPACE "BehaviorTreeEditor"
 
 UBehaviorTreeGraphNode_CompositeDecorator::UBehaviorTreeGraphNode_CompositeDecorator(const FObjectInitializer& ObjectInitializer) : Super(ObjectInitializer)
 {
@@ -37,9 +37,9 @@ FText UBehaviorTreeGraphNode_CompositeDecorator::GetNodeTitle(ENodeTitleType::Ty
 	return FText::FromString(CompositeName.Len() ? CompositeName : GetNodeTypeDescription());
 }
 
-FString UBehaviorTreeGraphNode_CompositeDecorator::GetDescription() const
+FText UBehaviorTreeGraphNode_CompositeDecorator::GetDescription() const
 {
-	return CachedDescription;
+	return FText::FromString(CachedDescription);
 }
 
 void UBehaviorTreeGraphNode_CompositeDecorator::PostPlacedNewNode()
diff --git a/Engine/Source/Editor/BehaviorTreeEditor/Private/BehaviorTreeGraphNode_Decorator.cpp b/Engine/Source/Editor/BehaviorTreeEditor/Private/BehaviorTreeGraphNode_Decorator.cpp
--- a/Engine/Source/Editor/BehaviorTreeEditor/Private/BehaviorTreeGraphNode_Decorator.cpp
+++ b/Engine/Source/Editor/BehaviorTreeEditor/Private/BehaviorTreeGraphNode_Decorator.cpp
@@ -6,6 +6,7 @@
 UBehaviorTreeGraphNode_Decorator::UBehaviorTreeGraphNode_Decorator(const FObjectInitializer& ObjectInitializer)
 	: Super(ObjectInitializer)
 {
+	bIsSubNode = true;
 }
 
 void UBehaviorTreeGraphNode_Decorator::AllocateDefaultPins()
@@ -22,17 +23,15 @@ FText UBehaviorTreeGraphNode_Decorator::GetNodeTitle(ENodeTitleType::Type TitleT
 	}
 	else if (!ClassData.GetClassName().IsEmpty())
 	{
-		return NSLOCTEXT("BehaviorTreeGraphNode", "UnknownNodeClass", "Can't load class!");
+		FString StoredClassName = ClassData.GetClassName();
+		StoredClassName.RemoveFromEnd(TEXT("_C"));
+
+		return FText::Format(NSLOCTEXT("AIGraph", "NodeClassError", "Class {0} not found, make sure it's saved!"), FText::FromString(StoredClassName));
 	}
 
 	return Super::GetNodeTitle(TitleType);
 }
 
-bool UBehaviorTreeGraphNode_Decorator::IsSubNode() const
-{
-	return true;
-}
-
 void UBehaviorTreeGraphNode_Decorator::CollectDecoratorData(TArray<UBTDecorator*>& NodeInstances, TArray<FBTDecoratorLogic>& Operations) const
 {
 	if (NodeInstance)
diff --git a/Engine/Source/Editor/BehaviorTreeEditor/Private/BehaviorTreeGraphNode_Root.cpp b/Engine/Source/Editor/BehaviorTreeEditor/Private/BehaviorTreeGraphNode_Root.cpp
--- a/Engine/Source/Editor/BehaviorTreeEditor/Private/BehaviorTreeGraphNode_Root.cpp
+++ b/Engine/Source/Editor/BehaviorTreeEditor/Private/BehaviorTreeGraphNode_Root.cpp
@@ -6,6 +6,7 @@
 
 UBehaviorTreeGraphNode_Root::UBehaviorTreeGraphNode_Root(const FObjectInitializer& ObjectInitializer) : Super(ObjectInitializer)
 {
+	bIsReadOnly = true;
 }
 
 void UBehaviorTreeGraphNode_Root::PostPlacedNewNode()
@@ -32,7 +33,7 @@ void UBehaviorTreeGraphNode_Root::AllocateDefaultPins()
 
 FText UBehaviorTreeGraphNode_Root::GetNodeTitle(ENodeTitleType::Type TitleType) const
 {
-	return NSLOCTEXT("BehaviorTreeGraphNode", "Root", "ROOT");
+	return NSLOCTEXT("BehaviorTreeEditor", "Root", "ROOT");
 }
 
 FName UBehaviorTreeGraphNode_Root::GetNameIcon() const
@@ -62,9 +63,9 @@ void UBehaviorTreeGraphNode_Root::PostEditUndo()
 	UpdateBlackboard();
 }
 
-FString	UBehaviorTreeGraphNode_Root::GetDescription() const
+FText UBehaviorTreeGraphNode_Root::GetDescription() const
 {
-	return *GetNameSafe(BlackboardAsset);
+	return FText::FromString(GetNameSafe(BlackboardAsset));
 }
 
 void UBehaviorTreeGraphNode_Root::UpdateBlackboard()
diff --git a/Engine/Source/Editor/BehaviorTreeEditor/Private/BehaviorTreeGraphNode_Service.cpp b/Engine/Source/Editor/BehaviorTreeEditor/Private/BehaviorTreeGraphNode_Service.cpp
--- a/Engine/Source/Editor/BehaviorTreeEditor/Private/BehaviorTreeGraphNode_Service.cpp
+++ b/Engine/Source/Editor/BehaviorTreeEditor/Private/BehaviorTreeGraphNode_Service.cpp
@@ -6,6 +6,7 @@
 UBehaviorTreeGraphNode_Service::UBehaviorTreeGraphNode_Service(const FObjectInitializer& ObjectInitializer)
 	: Super(ObjectInitializer)
 {
+	bIsSubNode = true;
 }
 
 void UBehaviorTreeGraphNode_Service::AllocateDefaultPins()
@@ -23,13 +24,11 @@ FText UBehaviorTreeGraphNode_Service::GetNodeTitle(ENodeTitleType::Type TitleTyp
 	}
 	else if (!ClassData.GetClassName().IsEmpty())
 	{
-		return NSLOCTEXT("BehaviorTreeGraphNode", "UnknownNodeClass", "Can't load class!");
+		FString StoredClassName = ClassData.GetClassName();
+		StoredClassName.RemoveFromEnd(TEXT("_C"));
+
+		return FText::Format(NSLOCTEXT("AIGraph", "NodeClassError", "Class {0} not found, make sure it's saved!"), FText::FromString(StoredClassName));
 	}
 
 	return Super::GetNodeTitle(TitleType);
 }
-
-bool UBehaviorTreeGraphNode_Service::IsSubNode() const
-{
-	return true;
-}
diff --git a/Engine/Source/Editor/BehaviorTreeEditor/Private/BehaviorTreeGraphNode_SimpleParallel.cpp b/Engine/Source/Editor/BehaviorTreeEditor/Private/BehaviorTreeGraphNode_SimpleParallel.cpp
--- a/Engine/Source/Editor/BehaviorTreeEditor/Private/BehaviorTreeGraphNode_SimpleParallel.cpp
+++ b/Engine/Source/Editor/BehaviorTreeEditor/Private/BehaviorTreeGraphNode_SimpleParallel.cpp
@@ -3,7 +3,7 @@
 #include "BehaviorTreeEditorPrivatePCH.h"
 #include "Internationalization.h" // for LOCTEXT()
 
-#define LOCTEXT_NAMESPACE "BehaviorTreeSchema"
+#define LOCTEXT_NAMESPACE "BehaviorTreeEditor"
 
 UBehaviorTreeGraphNode_SimpleParallel::UBehaviorTreeGraphNode_SimpleParallel(const FObjectInitializer& ObjectInitializer) : Super(ObjectInitializer)
 {
diff --git a/Engine/Source/Editor/BehaviorTreeEditor/Private/BehaviorTreeGraphNode_SubtreeTask.cpp b/Engine/Source/Editor/BehaviorTreeEditor/Private/BehaviorTreeGraphNode_SubtreeTask.cpp
--- a/Engine/Source/Editor/BehaviorTreeEditor/Private/BehaviorTreeGraphNode_SubtreeTask.cpp
+++ b/Engine/Source/Editor/BehaviorTreeEditor/Private/BehaviorTreeGraphNode_SubtreeTask.cpp
@@ -20,7 +20,7 @@ bool UBehaviorTreeGraphNode_SubtreeTask::UpdateInjectedNodes()
 	}
 
 	UBehaviorTreeGraph* MyGraph = MyNode->GetSubtreeAsset() ? Cast<UBehaviorTreeGraph>(MyNode->GetSubtreeAsset()->BTGraph) : NULL;
-	int32 MyVersion = MyGraph ? MyGraph->GraphVersion : 0;
+	int32 MyVersion = MyGraph ? MyGraph->ModCounter : 0;
 	FString MyPath = MyNode->GetSubtreeAsset() ? MyNode->GetSubtreeAsset()->GetName() : FString();
 
 	if (MyPath == SubtreePath && MyVersion == SubtreeVersion)
@@ -77,6 +77,7 @@ bool UBehaviorTreeGraphNode_SubtreeTask::UpdateInjectedNodes()
 
 					InjectedNode->ParentNode = this;
 					InjectedNode->bInjectedNode = true;
+					InjectedNode->bIsReadOnly = true;
 
 					UBTDecorator* InjectedInstance = Cast<UBTDecorator>(InjectedNode->NodeInstance);
 					if (InjectedInstance)
@@ -107,6 +108,7 @@ bool UBehaviorTreeGraphNode_SubtreeTask::UpdateInjectedNodes()
 						}
 					}
 
+					SubNodes.Add(InjectedNode);
 					Decorators.Add(InjectedNode);
 				}
 			}
diff --git a/Engine/Source/Editor/BehaviorTreeEditor/Private/BehaviorTreeGraphNode_Task.cpp b/Engine/Source/Editor/BehaviorTreeEditor/Private/BehaviorTreeGraphNode_Task.cpp
--- a/Engine/Source/Editor/BehaviorTreeEditor/Private/BehaviorTreeGraphNode_Task.cpp
+++ b/Engine/Source/Editor/BehaviorTreeEditor/Private/BehaviorTreeGraphNode_Task.cpp
@@ -20,7 +20,10 @@ FText UBehaviorTreeGraphNode_Task::GetNodeTitle(ENodeTitleType::Type TitleType)
 	}
 	else if (!ClassData.GetClassName().IsEmpty())
 	{
-		return NSLOCTEXT("BehaviorTreeGraphNode", "UnknownNodeClass", "Can't load class!");
+		FString StoredClassName = ClassData.GetClassName();
+		StoredClassName.RemoveFromEnd(TEXT("_C"));
+
+		return FText::Format(NSLOCTEXT("AIGraph", "NodeClassError", "Class {0} not found, make sure it's saved!"), FText::FromString(StoredClassName));
 	}
 
 	return Super::GetNodeTitle(TitleType);
diff --git a/Engine/Source/Editor/BehaviorTreeEditor/Private/EdGraphSchema_BehaviorTree.cpp b/Engine/Source/Editor/BehaviorTreeEditor/Private/EdGraphSchema_BehaviorTree.cpp
--- a/Engine/Source/Editor/BehaviorTreeEditor/Private/EdGraphSchema_BehaviorTree.cpp
+++ b/Engine/Source/Editor/BehaviorTreeEditor/Private/EdGraphSchema_BehaviorTree.cpp
@@ -15,14 +15,7 @@
 #include "BehaviorTree/Composites/BTComposite_SimpleParallel.h"
 #include "GenericCommands.h"
 
-#define LOCTEXT_NAMESPACE "BehaviorTreeSchema"
-#define SNAP_GRID (16) // @todo ensure this is the same as SNodePanel::GetSnapGridSize()
-
-namespace 
-{
-	// Maximum distance a drag can be off a node edge to require 'push off' from node
-	const int32 NodeDistance = 60;
-}
+#define LOCTEXT_NAMESPACE "BehaviorTreeEditor"
 
 UEdGraphNode* FBehaviorTreeSchemaAction_AutoArrange::PerformAction(class UEdGraph* ParentGraph, UEdGraphPin* FromPin, const FVector2D Location, bool bSelectNewNode)
 {
@@ -35,129 +28,10 @@ UEdGraphNode* FBehaviorTreeSchemaAction_AutoArrange::PerformAction(class UEdGrap
 	return NULL;
 }
 
-UEdGraphNode* FBehaviorTreeSchemaAction_NewNode::PerformAction(class UEdGraph* ParentGraph, UEdGraphPin* FromPin, const FVector2D Location, bool bSelectNewNode)
-{
-	UEdGraphNode* ResultNode = NULL;
-
-	// If there is a template, we actually use it
-	if (NodeTemplate != NULL)
-	{
-		const FScopedTransaction Transaction(LOCTEXT("AddNode", "Add Node"));
-		ParentGraph->Modify();
-		if (FromPin)
-		{
-			FromPin->Modify();
-		}
-
-		NodeTemplate->SetFlags(RF_Transactional);
-
-		// set outer to be the graph so it doesn't go away
-		NodeTemplate->Rename(NULL, ParentGraph, REN_NonTransactional);
-		ParentGraph->AddNode(NodeTemplate, true);
-
-		NodeTemplate->CreateNewGuid();
-		NodeTemplate->PostPlacedNewNode();
-		NodeTemplate->AllocateDefaultPins();
-		NodeTemplate->AutowireNewNode(FromPin);
-
-		// For input pins, new node will generally overlap node being dragged off
-		// Work out if we want to visually push away from connected node
-		int32 XLocation = Location.X;
-		if (FromPin && FromPin->Direction == EGPD_Input)
-		{
-			UEdGraphNode* PinNode = FromPin->GetOwningNode();
-			const float XDelta = FMath::Abs(PinNode->NodePosX - Location.X);
-
-			if (XDelta < NodeDistance)
-			{
-				// Set location to edge of current node minus the max move distance
-				// to force node to push off from connect node enough to give selection handle
-				XLocation = PinNode->NodePosX - NodeDistance;
-			}
-		}
-
-		NodeTemplate->NodePosX = XLocation;
-		NodeTemplate->NodePosY = Location.Y;
-		NodeTemplate->SnapToGrid(SNAP_GRID);
-
-		ResultNode = NodeTemplate;
-	}
-
-	return ResultNode;
-}
-
-UEdGraphNode* FBehaviorTreeSchemaAction_NewNode::PerformAction(class UEdGraph* ParentGraph, TArray<UEdGraphPin*>& FromPins, const FVector2D Location, bool bSelectNewNode) 
-{
-	UEdGraphNode* ResultNode = NULL;
-	if (FromPins.Num() > 0)
-	{
-		ResultNode = PerformAction(ParentGraph, FromPins[0], Location);
-
-		// Try autowiring the rest of the pins
-		for (int32 Index = 1; Index < FromPins.Num(); ++Index)
-		{
-			ResultNode->AutowireNewNode(FromPins[Index]);
-		}
-	}
-	else
-	{
-		ResultNode = PerformAction(ParentGraph, NULL, Location, bSelectNewNode);
-	}
-
-	return ResultNode;
-}
-
-void FBehaviorTreeSchemaAction_NewNode::AddReferencedObjects( FReferenceCollector& Collector )
-{
-	FEdGraphSchemaAction::AddReferencedObjects( Collector );
-
-	// These don't get saved to disk, but we want to make sure the objects don't get GC'd while the action array is around
-	Collector.AddReferencedObject( NodeTemplate );
-}
-
-
-
-
-UEdGraphNode* FBehaviorTreeSchemaAction_NewSubNode::PerformAction(class UEdGraph* ParentGraph, UEdGraphPin* FromPin, const FVector2D Location, bool bSelectNewNode)
-{
-	ParentNode->AddSubNode(NodeTemplate,ParentGraph);
-	return NULL;
-}
-
-UEdGraphNode* FBehaviorTreeSchemaAction_NewSubNode::PerformAction(class UEdGraph* ParentGraph, TArray<UEdGraphPin*>& FromPins, const FVector2D Location, bool bSelectNewNode) 
-{
-	return PerformAction(ParentGraph, NULL, Location, bSelectNewNode);
-}
-
-void FBehaviorTreeSchemaAction_NewSubNode::AddReferencedObjects( FReferenceCollector& Collector )
-{
-	FEdGraphSchemaAction::AddReferencedObjects( Collector );
-
-	// These don't get saved to disk, but we want to make sure the objects don't get GC'd while the action array is around
-	Collector.AddReferencedObject( NodeTemplate );
-}
-//////////////////////////////////////////////////////////////////////////
-
 UEdGraphSchema_BehaviorTree::UEdGraphSchema_BehaviorTree(const FObjectInitializer& ObjectInitializer) : Super(ObjectInitializer)
 {
 }
 
-TSharedPtr<FBehaviorTreeSchemaAction_NewNode> AddNewNodeAction(FGraphActionListBuilderBase& ContextMenuBuilder, const FString& Category, const FText& MenuDesc, const FString& Tooltip)
-{
-	TSharedPtr<FBehaviorTreeSchemaAction_NewNode> NewAction = TSharedPtr<FBehaviorTreeSchemaAction_NewNode>(new FBehaviorTreeSchemaAction_NewNode(Category, MenuDesc, Tooltip, 0));
-
-	ContextMenuBuilder.AddAction( NewAction );
-
-	return NewAction;
-}
-
-TSharedPtr<FBehaviorTreeSchemaAction_NewSubNode> AddNewSubNodeAction(FGraphActionListBuilderBase& ContextMenuBuilder, const FString& Category, const FText& MenuDesc, const FString& Tooltip)
-{
-	TSharedPtr<FBehaviorTreeSchemaAction_NewSubNode> NewAction = TSharedPtr<FBehaviorTreeSchemaAction_NewSubNode>(new FBehaviorTreeSchemaAction_NewSubNode(Category, MenuDesc, Tooltip, 0));
-	ContextMenuBuilder.AddAction( NewAction );
-	return NewAction;
-}
-
 void UEdGraphSchema_BehaviorTree::CreateDefaultNodesForGraph(UEdGraph& Graph) const
 {
 	FGraphNodeCreator<UBehaviorTreeGraphNode_Root> NodeCreator(Graph);
@@ -166,66 +40,35 @@ void UEdGraphSchema_BehaviorTree::CreateDefaultNodesForGraph(UEdGraph& Graph) co
 	SetNodeMetaData(MyNode, FNodeMetadata::DefaultGraphNode);
 }
 
-FString UEdGraphSchema_BehaviorTree::GetShortTypeName(const UObject* Ob) const
+void UEdGraphSchema_BehaviorTree::GetGraphNodeContextActions(FGraphContextMenuBuilder& ContextMenuBuilder, int32 SubNodeFlags) const
 {
-	if (Ob->GetClass()->HasAnyClassFlags(CLASS_CompiledFromBlueprint))
-	{
-		return Ob->GetClass()->GetName().LeftChop(2);
-	}
+	Super::GetGraphNodeContextActions(ContextMenuBuilder, SubNodeFlags);
 
-	FString TypeDesc = Ob->GetClass()->GetName();
-	const int32 ShortNameIdx = TypeDesc.Find(TEXT("_"));
-	if (ShortNameIdx != INDEX_NONE)
+	if (SubNodeFlags == ESubNode::Decorator)
 	{
-		TypeDesc = TypeDesc.Mid(ShortNameIdx + 1);
+		const FString& Category = UBehaviorTreeGraphNode_CompositeDecorator::StaticClass()->GetMetaData(TEXT("Category"));
+		UEdGraph* Graph = (UEdGraph*)ContextMenuBuilder.CurrentGraph;
+		UBehaviorTreeGraphNode_CompositeDecorator* OpNode = NewObject<UBehaviorTreeGraphNode_CompositeDecorator>(Graph);
+		TSharedPtr<FAISchemaAction_NewSubNode> AddOpAction = UAIGraphSchema::AddNewSubNodeAction(ContextMenuBuilder, Category, FText::FromString(OpNode->GetNodeTypeDescription()), "");
+		AddOpAction->ParentNode = Cast<UBehaviorTreeGraphNode>(ContextMenuBuilder.SelectedObjects[0]);
+		AddOpAction->NodeTemplate = OpNode;
 	}
-
-	return TypeDesc;
 }
 
-void UEdGraphSchema_BehaviorTree::GetGraphNodeContextActions(FGraphContextMenuBuilder& ContextMenuBuilder, ESubNode::Type SubNodeType) const
+void UEdGraphSchema_BehaviorTree::GetSubNodeClasses(int32 SubNodeFlags, TArray<FGraphNodeClassData>& ClassData, UClass*& GraphNodeClass) const
 {
-	TArray<FClassData> NodeClasses;
-	UEdGraph* Graph = (UEdGraph*)ContextMenuBuilder.CurrentGraph;
+	FBehaviorTreeEditorModule& EditorModule = FModuleManager::GetModuleChecked<FBehaviorTreeEditorModule>(TEXT("BehaviorTreeEditor"));
+	FGraphNodeClassHelper* ClassCache = EditorModule.GetClassCache().Get();
 
-	if (SubNodeType == ESubNode::Decorator)
+	if (SubNodeFlags == ESubNode::Decorator)
 	{
-		FClassBrowseHelper::GatherClasses(UBTDecorator::StaticClass(), NodeClasses);
-
-		{
-			const FString& Category = UBehaviorTreeGraphNode_CompositeDecorator::StaticClass()->GetMetaData(TEXT("Category"));
-			UBehaviorTreeGraphNode_CompositeDecorator* OpNode = NewObject<UBehaviorTreeGraphNode_CompositeDecorator>(Graph);
-			TSharedPtr<FBehaviorTreeSchemaAction_NewSubNode> AddOpAction = AddNewSubNodeAction(ContextMenuBuilder, Category, FText::FromString(OpNode->GetNodeTypeDescription()), "");
-			AddOpAction->ParentNode = Cast<UBehaviorTreeGraphNode>(ContextMenuBuilder.SelectedObjects[0]);
-			AddOpAction->NodeTemplate = OpNode;
-		}
-
-		for (const auto& NodeClass : NodeClasses)
-		{
-			const FText NodeTypeName = FText::FromString(FName::NameToDisplayString(NodeClass.ToString(), false));
-
-			UBehaviorTreeGraphNode* OpNode = NewObject<UBehaviorTreeGraphNode_Decorator>(Graph);
-			OpNode->ClassData = NodeClass;
-
-			TSharedPtr<FBehaviorTreeSchemaAction_NewSubNode> AddOpAction = AddNewSubNodeAction(ContextMenuBuilder, NodeClass.GetCategory(), NodeTypeName, "");
-			AddOpAction->ParentNode = Cast<UBehaviorTreeGraphNode>(ContextMenuBuilder.SelectedObjects[0]);
-			AddOpAction->NodeTemplate = OpNode;
-		}
+		ClassCache->GatherClasses(UBTDecorator::StaticClass(), ClassData);
+		GraphNodeClass = UBehaviorTreeGraphNode_Decorator::StaticClass();
 	}
-	else if (SubNodeType == ESubNode::Service)
+	else
 	{
-		FClassBrowseHelper::GatherClasses(UBTService::StaticClass(), NodeClasses);
-		for (const auto& NodeClass : NodeClasses)
-		{
-			const FText NodeTypeName = FText::FromString(FName::NameToDisplayString(NodeClass.ToString(), false));
-
-			UBehaviorTreeGraphNode* OpNode = NewObject<UBehaviorTreeGraphNode_Service>(Graph);
-			OpNode->ClassData = NodeClass;
-
-			TSharedPtr<FBehaviorTreeSchemaAction_NewSubNode> AddOpAction = AddNewSubNodeAction(ContextMenuBuilder, NodeClass.GetCategory(), NodeTypeName, "");
-			AddOpAction->ParentNode = Cast<UBehaviorTreeGraphNode>(ContextMenuBuilder.SelectedObjects[0]);
-			AddOpAction->NodeTemplate = OpNode;
-		}
+		ClassCache->GatherClasses(UBTService::StaticClass(), ClassData);
+		GraphNodeClass = UBehaviorTreeGraphNode_Service::StaticClass();
 	}
 }
 
@@ -241,25 +84,30 @@ void UEdGraphSchema_BehaviorTree::GetGraphContextActions(FGraphContextMenuBuilde
 	const bool bAllowComposites = bNoParent || !bOnlyTasks || bOnlyComposites;
 	const bool bAllowTasks = bNoParent || !bOnlyComposites || bOnlyTasks;
 
+	FBehaviorTreeEditorModule& EditorModule = FModuleManager::GetModuleChecked<FBehaviorTreeEditorModule>(TEXT("BehaviorTreeEditor"));
+	FGraphNodeClassHelper* ClassCache = EditorModule.GetClassCache().Get();
+
 	if (bAllowComposites)
 	{
 		FCategorizedGraphActionListBuilder CompositesBuilder(TEXT("Composites"));
 
-		TArray<FClassData> NodeClasses;
-		FClassBrowseHelper::GatherClasses(UBTCompositeNode::StaticClass(), NodeClasses);
+		TArray<FGraphNodeClassData> NodeClasses;
+		ClassCache->GatherClasses(UBTCompositeNode::StaticClass(), NodeClasses);
 
 		const FString ParallelClassName = UBTComposite_SimpleParallel::StaticClass()->GetName();
-
 		for (const auto& NodeClass : NodeClasses)
 		{
 			const FText NodeTypeName = FText::FromString(FName::NameToDisplayString(NodeClass.ToString(), false));
 
-			TSharedPtr<FBehaviorTreeSchemaAction_NewNode> AddOpAction = AddNewNodeAction(CompositesBuilder, NodeClass.GetCategory(), NodeTypeName, "");
+			TSharedPtr<FAISchemaAction_NewNode> AddOpAction = UAIGraphSchema::AddNewNodeAction(CompositesBuilder, NodeClass.GetCategory(), NodeTypeName, "");
 
-			UBehaviorTreeGraphNode* OpNode = (NodeClass.GetClassName() == ParallelClassName) ?
-				NewObject<UBehaviorTreeGraphNode_SimpleParallel>(ContextMenuBuilder.OwnerOfTemporaries) :				
-				NewObject<UBehaviorTreeGraphNode_Composite>(ContextMenuBuilder.OwnerOfTemporaries);
+			UClass* GraphNodeClass = UBehaviorTreeGraphNode_Composite::StaticClass();
+			if (NodeClass.GetClassName() == ParallelClassName)
+			{
+				GraphNodeClass = UBehaviorTreeGraphNode_SimpleParallel::StaticClass();
+			}
 
+			UBehaviorTreeGraphNode* OpNode = NewObject<UBehaviorTreeGraphNode>(ContextMenuBuilder.OwnerOfTemporaries, GraphNodeClass);
 			OpNode->ClassData = NodeClass;
 			AddOpAction->NodeTemplate = OpNode;
 		}
@@ -271,16 +119,16 @@ void UEdGraphSchema_BehaviorTree::GetGraphContextActions(FGraphContextMenuBuilde
 	{
 		FCategorizedGraphActionListBuilder TasksBuilder(TEXT("Tasks"));
 
-		TArray<FClassData> NodeClasses;
-		FClassBrowseHelper::GatherClasses(UBTTaskNode::StaticClass(), NodeClasses);
+		TArray<FGraphNodeClassData> NodeClasses;
+		ClassCache->GatherClasses(UBTTaskNode::StaticClass(), NodeClasses);
 
 		for (const auto& NodeClass : NodeClasses)
 		{
 			const FText NodeTypeName = FText::FromString(FName::NameToDisplayString(NodeClass.ToString(), false));
 
-			TSharedPtr<FBehaviorTreeSchemaAction_NewNode> AddOpAction = AddNewNodeAction(TasksBuilder, NodeClass.GetCategory(), NodeTypeName, "");
-			UClass* GraphNodeClass = UBehaviorTreeGraphNode_Task::StaticClass();
+			TSharedPtr<FAISchemaAction_NewNode> AddOpAction = UAIGraphSchema::AddNewNodeAction(TasksBuilder, NodeClass.GetCategory(), NodeTypeName, "");
 			
+			UClass* GraphNodeClass = UBehaviorTreeGraphNode_Task::StaticClass();
 			if (NodeClass.GetClassName() == UBTTask_RunBehavior::StaticClass()->GetName())
 			{
 				GraphNodeClass = UBehaviorTreeGraphNode_SubtreeTask::StaticClass();
@@ -296,134 +144,36 @@ void UEdGraphSchema_BehaviorTree::GetGraphContextActions(FGraphContextMenuBuilde
 	
 	if (bNoParent)
 	{
-		TSharedPtr<FBehaviorTreeSchemaAction_AutoArrange> Action = TSharedPtr<FBehaviorTreeSchemaAction_AutoArrange>(new FBehaviorTreeSchemaAction_AutoArrange(FString(), LOCTEXT("AutoArrange", "Auto Arrange"), FString(), 0));
+		TSharedPtr<FBehaviorTreeSchemaAction_AutoArrange> Action = TSharedPtr<FBehaviorTreeSchemaAction_AutoArrange>(
+			new FBehaviorTreeSchemaAction_AutoArrange(FString(), LOCTEXT("AutoArrange", "Auto Arrange"), FString(), 0)
+			);
+
 		ContextMenuBuilder.AddAction(Action);
 	}
 }
 
 void UEdGraphSchema_BehaviorTree::GetContextMenuActions(const UEdGraph* CurrentGraph, const UEdGraphNode* InGraphNode, const UEdGraphPin* InGraphPin, class FMenuBuilder* MenuBuilder, bool bIsDebugging) const
 {
-	if (InGraphPin)
-	{
-		MenuBuilder->BeginSection("BehaviorTreeGraphSchemaPinActions", LOCTEXT("PinActionsMenuHeader", "Pin Actions"));
-		{
-			// Only display the 'Break Links' option if there is a link to break!
-			if (InGraphPin->LinkedTo.Num() > 0)
-			{
-				MenuBuilder->AddMenuEntry( FGraphEditorCommands::Get().BreakPinLinks );
-
-				// add sub menu for break link to
-				if(InGraphPin->LinkedTo.Num() > 1)
-				{
-					MenuBuilder->AddSubMenu(
-						LOCTEXT("BreakLinkTo", "Break Link To..." ),
-						LOCTEXT("BreakSpecificLinks", "Break a specific link..." ),
-						FNewMenuDelegate::CreateUObject( (UEdGraphSchema_BehaviorTree*const)this, &UEdGraphSchema_BehaviorTree::GetBreakLinkToSubMenuActions, const_cast<UEdGraphPin*>(InGraphPin)));
-				}
-				else
-				{
-					((UEdGraphSchema_BehaviorTree*const)this)->GetBreakLinkToSubMenuActions(*MenuBuilder, const_cast<UEdGraphPin*>(InGraphPin));
-				}
-			}
-		}
-		MenuBuilder->EndSection();
-	}
-	else if (InGraphNode)
+	if (InGraphNode && !InGraphPin)
 	{
 		const UBehaviorTreeGraphNode* BTGraphNode = Cast<const UBehaviorTreeGraphNode>(InGraphNode);
 		if (BTGraphNode && BTGraphNode->CanPlaceBreakpoints())
 		{
 			MenuBuilder->BeginSection("EdGraphSchemaBreakpoints", LOCTEXT("BreakpointsHeader", "Breakpoints"));
 			{
-				MenuBuilder->AddMenuEntry( FGraphEditorCommands::Get().ToggleBreakpoint );
-				MenuBuilder->AddMenuEntry( FGraphEditorCommands::Get().AddBreakpoint );
-				MenuBuilder->AddMenuEntry( FGraphEditorCommands::Get().RemoveBreakpoint );
-				MenuBuilder->AddMenuEntry( FGraphEditorCommands::Get().EnableBreakpoint );
-				MenuBuilder->AddMenuEntry( FGraphEditorCommands::Get().DisableBreakpoint );
+				MenuBuilder->AddMenuEntry(FGraphEditorCommands::Get().ToggleBreakpoint);
+				MenuBuilder->AddMenuEntry(FGraphEditorCommands::Get().AddBreakpoint);
+				MenuBuilder->AddMenuEntry(FGraphEditorCommands::Get().RemoveBreakpoint);
+				MenuBuilder->AddMenuEntry(FGraphEditorCommands::Get().EnableBreakpoint);
+				MenuBuilder->AddMenuEntry(FGraphEditorCommands::Get().DisableBreakpoint);
 			}
 			MenuBuilder->EndSection();
 		}
-
-		MenuBuilder->BeginSection("BehaviorTreeGraphSchemaNodeActions", LOCTEXT("ClassActionsMenuHeader", "Node Actions"));
-		{
-			MenuBuilder->AddMenuEntry( FGenericCommands::Get().Delete );
-			MenuBuilder->AddMenuEntry( FGenericCommands::Get().Cut );
-			MenuBuilder->AddMenuEntry( FGenericCommands::Get().Copy );
-			MenuBuilder->AddMenuEntry( FGenericCommands::Get().Duplicate );
-
-			MenuBuilder->AddMenuEntry(FGraphEditorCommands::Get().BreakNodeLinks);
-		}
-		MenuBuilder->EndSection();
 	}
 
 	Super::GetContextMenuActions(CurrentGraph, InGraphNode, InGraphPin, MenuBuilder, bIsDebugging);
 }
 
-void UEdGraphSchema_BehaviorTree::GetBreakLinkToSubMenuActions( class FMenuBuilder& MenuBuilder, UEdGraphPin* InGraphPin )
-{
-	// Make sure we have a unique name for every entry in the list
-	TMap< FString, uint32 > LinkTitleCount;
-
-	// Add all the links we could break from
-	for(TArray<class UEdGraphPin*>::TConstIterator Links(InGraphPin->LinkedTo); Links; ++Links)
-	{
-		UEdGraphPin* Pin = *Links;
-		FString TitleString = Pin->GetOwningNode()->GetNodeTitle(ENodeTitleType::ListView).ToString();
-		FText Title = FText::FromString( TitleString );
-		if ( Pin->PinName != TEXT("") )
-		{
-			TitleString = FString::Printf(TEXT("%s (%s)"), *TitleString, *Pin->PinName);
-
-			// Add name of connection if possible
-			FFormatNamedArguments Args;
-			Args.Add( TEXT("NodeTitle"), Title );
-			Args.Add( TEXT("PinName"), Pin->GetDisplayName() );
-			Title = FText::Format( LOCTEXT("BreakDescPin", "{NodeTitle} ({PinName})"), Args );
-		}
-
-		uint32 &Count = LinkTitleCount.FindOrAdd( TitleString );
-
-		FText Description;
-		FFormatNamedArguments Args;
-		Args.Add( TEXT("NodeTitle"), Title );
-		Args.Add( TEXT("NumberOfNodes"), Count );
-
-		if(Count == 0)
-		{
-			Description = FText::Format( LOCTEXT("BreakDesc", "Break link to {NodeTitle}"), Args );
-		}
-		else
-		{
-			Description = FText::Format( LOCTEXT("BreakDescMulti", "Break link to {NodeTitle} ({NumberOfNodes})"), Args );
-		}
-		++Count;
-
-		MenuBuilder.AddMenuEntry( Description, Description, FSlateIcon(), FUIAction(
-			FExecuteAction::CreateUObject(this, &UEdGraphSchema_BehaviorTree::BreakSinglePinLink, const_cast< UEdGraphPin* >(InGraphPin), *Links) ) );
-	}
-}
-
-void UEdGraphSchema_BehaviorTree::BreakNodeLinks(UEdGraphNode& TargetNode) const
-{
-	const FScopedTransaction Transaction( NSLOCTEXT("UnrealEd", "GraphEd_BreakNodeLinks", "Break Node Links") );
-
-	Super::BreakNodeLinks(TargetNode);
-}
-
-void UEdGraphSchema_BehaviorTree::BreakPinLinks(UEdGraphPin& TargetPin, bool bSendsNodeNotification) const
-{
-	const FScopedTransaction Transaction( NSLOCTEXT("UnrealEd", "GraphEd_BreakPinLinks", "Break Pin Links") );
-
-	Super::BreakPinLinks(TargetPin, bSendsNodeNotification);
-}
-
-void UEdGraphSchema_BehaviorTree::BreakSinglePinLink(UEdGraphPin* SourcePin, UEdGraphPin* TargetPin)
-{
-	const FScopedTransaction Transaction( NSLOCTEXT("UnrealEd", "GraphEd_BreakSinglePinLink", "Break Pin Link") );
-
-	Super::BreakSinglePinLink(SourcePin, TargetPin);
-}
-
 const FPinConnectionResponse UEdGraphSchema_BehaviorTree::CanCreateConnection(const UEdGraphPin* PinA, const UEdGraphPin* PinB) const
 {
 	// Make sure the pins are not on the same node
@@ -585,17 +335,19 @@ const FPinConnectionResponse UEdGraphSchema_BehaviorTree::CanMergeNodes(const UE
 
 			if (BTNodeB && BTNodeB->bInjectedNode)
 			{
+				const UBehaviorTreeGraphNode* ParentNodeB = Cast<const UBehaviorTreeGraphNode>(BTNodeB->ParentNode);
+
 				int32 FirstInjectedIdx = INDEX_NONE;
-				for (int32 Idx = 0; Idx < BTNodeB->ParentNode->Decorators.Num(); Idx++)
+				for (int32 Idx = 0; Idx < ParentNodeB->Decorators.Num(); Idx++)
 				{
-					if (BTNodeB->ParentNode->Decorators[Idx]->bInjectedNode)
+					if (ParentNodeB->Decorators[Idx]->bInjectedNode)
 					{
 						FirstInjectedIdx = Idx;
 						break;
 					}
 				}
 
-				int32 NodeIdx = BTNodeB->ParentNode->Decorators.IndexOfByKey(BTNodeB);
+				int32 NodeIdx = ParentNodeB->Decorators.IndexOfByKey(BTNodeB);
 				if (NodeIdx != FirstInjectedIdx)
 				{
 					return FPinConnectionResponse(CONNECT_RESPONSE_DISALLOW, LOCTEXT("MergeInjectedNodeAtEnd", "Decorators must be placed above injected nodes!"));
@@ -629,18 +381,6 @@ FLinearColor UEdGraphSchema_BehaviorTree::GetPinTypeColor(const FEdGraphPinType&
 	return FColor::White;
 }
 
-bool UEdGraphSchema_BehaviorTree::ShouldHidePinDefaultValue(UEdGraphPin* Pin) const
-{
-	check(Pin != NULL);
-
-	if (Pin->bDefaultValueIsIgnored)
-	{
-		return true;
-	}
-
-	return false;
-}
-
 class FConnectionDrawingPolicy* UEdGraphSchema_BehaviorTree::CreateConnectionDrawingPolicy(int32 InBackLayerID, int32 InFrontLayerID, float InZoomFactor, const FSlateRect& InClippingRect, class FSlateWindowElementList& InDrawElements, class UEdGraph* InGraphObj) const
 {
 	return new FBehaviorTreeConnectionDrawingPolicy(InBackLayerID, InFrontLayerID, InZoomFactor, InClippingRect, InDrawElements, InGraphObj);
diff --git a/Engine/Source/Editor/BehaviorTreeEditor/Private/EdGraphSchema_BehaviorTreeDecorator.cpp b/Engine/Source/Editor/BehaviorTreeEditor/Private/EdGraphSchema_BehaviorTreeDecorator.cpp
--- a/Engine/Source/Editor/BehaviorTreeEditor/Private/EdGraphSchema_BehaviorTreeDecorator.cpp
+++ b/Engine/Source/Editor/BehaviorTreeEditor/Private/EdGraphSchema_BehaviorTreeDecorator.cpp
@@ -128,16 +128,19 @@ void UEdGraphSchema_BehaviorTreeDecorator::GetGraphContextActions(FGraphContextM
 	const UBehaviorTreeDecoratorGraphNode* ParentGraphNode = ContextMenuBuilder.FromPin 
 		? Cast<UBehaviorTreeDecoratorGraphNode>(ContextMenuBuilder.FromPin->GetOuter()) : NULL;
 
-	TArray<FClassData> NodeClasses;
-	FClassBrowseHelper::GatherClasses(UBTDecorator::StaticClass(), NodeClasses);
+	FBehaviorTreeEditorModule& EditorModule = FModuleManager::GetModuleChecked<FBehaviorTreeEditorModule>(TEXT("BehaviorTreeEditor"));
+	FGraphNodeClassHelper* ClassCache = EditorModule.GetClassCache().Get();
 
-	for (int32 i = 0; i < NodeClasses.Num(); i++)
+	TArray<FGraphNodeClassData> NodeClasses;
+	ClassCache->GatherClasses(UBTDecorator::StaticClass(), NodeClasses);
+
+	for (const auto& NodeClass : NodeClasses)
 	{
-		const FText NodeTypeName = FText::FromString(NodeClasses[i].ToString());
-		TSharedPtr<FDecoratorSchemaAction_NewNode> AddOpAction = AddNewDecoratorAction(ContextMenuBuilder, TEXT("Decorators"), NodeTypeName, "");
+		const FText NodeTypeName = FText::FromString(NodeClass.ToString());
+		TSharedPtr<FDecoratorSchemaAction_NewNode> AddOpAction = AddNewDecoratorAction(ContextMenuBuilder, NodeClass.GetCategory(), NodeTypeName, "");
 
 		UBehaviorTreeDecoratorGraphNode_Decorator* OpNode = NewObject<UBehaviorTreeDecoratorGraphNode_Decorator>(ContextMenuBuilder.OwnerOfTemporaries);
-		OpNode->ClassData = NodeClasses[i];
+		OpNode->ClassData = NodeClass;
 		AddOpAction->NodeTemplate = OpNode;
 	}
 
diff --git a/Engine/Source/Editor/BehaviorTreeEditor/Private/SGraphNode_BehaviorTree.cpp b/Engine/Source/Editor/BehaviorTreeEditor/Private/SGraphNode_BehaviorTree.cpp
--- a/Engine/Source/Editor/BehaviorTreeEditor/Private/SGraphNode_BehaviorTree.cpp
+++ b/Engine/Source/Editor/BehaviorTreeEditor/Private/SGraphNode_BehaviorTree.cpp
@@ -17,41 +17,12 @@
 #include "SInlineEditableTextBlock.h"
 #include "SCommentBubble.h"
 
-#define LOCTEXT_NAMESPACE "SGraphNode_BehaviorTree"
-
-TSharedRef<FDragNodeTimed> FDragNodeTimed::New(const TSharedRef<SGraphPanel>& InGraphPanel, const TSharedRef<SGraphNode>& InDraggedNode)
-{
-	TSharedRef<FDragNodeTimed> Operation = MakeShareable(new FDragNodeTimed);
-	Operation->StartTime = FPlatformTime::Seconds();
-	Operation->GraphPanel = InGraphPanel;
-	Operation->DraggedNodes.Add(InDraggedNode);
-	// adjust the decorator away from the current mouse location a small amount based on cursor size
-	Operation->DecoratorAdjust = FSlateApplication::Get().GetCursorSize();
-	Operation->Construct();
-	return Operation;
-}
-
-TSharedRef<FDragNodeTimed> FDragNodeTimed::New(const TSharedRef<SGraphPanel>& InGraphPanel, const TArray< TSharedRef<SGraphNode> >& InDraggedNodes)
-{
-	TSharedRef<FDragNodeTimed> Operation = MakeShareable(new FDragNodeTimed);
-	Operation->StartTime = FPlatformTime::Seconds();
-	Operation->GraphPanel = InGraphPanel;
-	Operation->DraggedNodes = InDraggedNodes;
-	Operation->DecoratorAdjust = FSlateApplication::Get().GetCursorSize();
-	Operation->Construct();
-	return Operation;
-}
-
-UBehaviorTreeGraphNode* FDragNodeTimed::GetDropTargetNode() const
-{
-	return Cast<UBehaviorTreeGraphNode>(GetHoveredNode());
-}
-
+#define LOCTEXT_NAMESPACE "BehaviorTreeEditor"
 
 /////////////////////////////////////////////////////
 // SBehaviorTreePin
 
-class SBehaviorTreePin : public SGraphPin
+class SBehaviorTreePin : public SGraphPinAI
 {
 public:
 	SLATE_BEGIN_ARGS(SBehaviorTreePin){}
@@ -59,47 +30,13 @@ public:
 
 	void Construct(const FArguments& InArgs, UEdGraphPin* InPin);
 protected:
-	// Begin SGraphPin interface
-	virtual TSharedRef<SWidget>	GetDefaultValueWidget() override;
-
 	/** @return The color that we should use to draw this pin */
 	virtual FSlateColor GetPinColor() const override;
-
-	// End SGraphPin interface
-
-	const FSlateBrush* GetPinBorder() const;
 };
 
 void SBehaviorTreePin::Construct(const FArguments& InArgs, UEdGraphPin* InPin)
 {
-	this->SetCursor( EMouseCursor::Default );
-
-	bShowLabel = true;
-	IsEditable = true;
-
-	GraphPinObj = InPin;
-	check(GraphPinObj != NULL);
-
-	const UEdGraphSchema* Schema = GraphPinObj->GetSchema();
-	check(Schema);
-
-	SBorder::Construct( SBorder::FArguments()
-		.BorderImage( this, &SBehaviorTreePin::GetPinBorder )
-		.BorderBackgroundColor( this, &SBehaviorTreePin::GetPinColor )
-		.OnMouseButtonDown( this, &SBehaviorTreePin::OnPinMouseDown )
-		.Cursor( this, &SBehaviorTreePin::GetPinCursor )
-		.Padding(FMargin(5.0f))
-		);
-}
-
-TSharedRef<SWidget>	SBehaviorTreePin::GetDefaultValueWidget()
-{
-	return SNew(STextBlock);
-}
-
-const FSlateBrush* SBehaviorTreePin::GetPinBorder() const
-{
-	return FEditorStyle::GetBrush(TEXT("Graph.StateNode.Body"));
+	SGraphPinAI::Construct(SGraphPinAI::FArguments(), InPin);
 }
 
 FSlateColor SBehaviorTreePin::GetPinColor() const
@@ -207,11 +144,7 @@ void SGraphNode_BehaviorTree::Construct(const FArguments& InArgs, UBehaviorTreeG
 	DebuggerStateCounter = INDEX_NONE;
 	bSuppressDebuggerTriggers = false;
 
-	this->GraphNode = InNode;
-
-	this->SetCursor(EMouseCursor::CardinalCross);
-	
-	this->UpdateGraphNode();
+	SGraphNodeAI::Construct(SGraphNodeAI::FArguments(), InNode);
 }
 
 void SGraphNode_BehaviorTree::AddDecorator(TSharedPtr<SGraphNode> DecoratorWidget)
@@ -220,7 +153,9 @@ void SGraphNode_BehaviorTree::AddDecorator(TSharedPtr<SGraphNode> DecoratorWidge
 	[
 		DecoratorWidget.ToSharedRef()
 	];
+
 	DecoratorWidgets.Add(DecoratorWidget);
+	AddSubNode(DecoratorWidget);
 }
 
 void SGraphNode_BehaviorTree::AddService(TSharedPtr<SGraphNode> ServiceWidget)
@@ -230,14 +165,16 @@ void SGraphNode_BehaviorTree::AddService(TSharedPtr<SGraphNode> ServiceWidget)
 		ServiceWidget.ToSharedRef()
 	];
 	ServicesWidgets.Add(ServiceWidget);
+	AddSubNode(ServiceWidget);
 }
 
 FSlateColor SGraphNode_BehaviorTree::GetBorderBackgroundColor() const
 {
 	UBehaviorTreeGraphNode* BTGraphNode = Cast<UBehaviorTreeGraphNode>(GraphNode);
+	UBehaviorTreeGraphNode* BTParentNode = BTGraphNode ? Cast<UBehaviorTreeGraphNode>(BTGraphNode->ParentNode) : nullptr;
 	const bool bIsInDebuggerActiveState = BTGraphNode && BTGraphNode->bDebuggerMarkCurrentlyActive;
 	const bool bIsInDebuggerPrevState = BTGraphNode && BTGraphNode->bDebuggerMarkPreviouslyActive;
-	const bool bSelectedSubNode = BTGraphNode && BTGraphNode->ParentNode && GetOwnerPanel()->SelectionManager.SelectedNodes.Contains(GraphNode);
+	const bool bSelectedSubNode = BTParentNode && GetOwnerPanel()->SelectionManager.SelectedNodes.Contains(GraphNode);
 	
 	UBTNode* NodeInstance = BTGraphNode ? Cast<UBTNode>(BTGraphNode->NodeInstance) : NULL;
 	const bool bIsDisconnected = NodeInstance && NodeInstance->GetExecutionIndex() == MAX_uint16;
@@ -245,8 +182,8 @@ FSlateColor SGraphNode_BehaviorTree::GetBorderBackgroundColor() const
 	const bool bIsRootDecorator = BTGraphNode && BTGraphNode->bRootLevel;
 	const bool bIsInjected = BTGraphNode && BTGraphNode->bInjectedNode;
 	const bool bIsBrokenWithParent = bIsService ? 
-		BTGraphNode->ParentNode != NULL && BTGraphNode->ParentNode->Services.Find(BTGraphNode) == INDEX_NONE ? true : false :
-		BTGraphNode->ParentNode != NULL && BTGraphNode->ParentNode->Decorators.Find(BTGraphNode) == INDEX_NONE ? true : 
+		BTParentNode && BTParentNode->Services.Find(BTGraphNode) == INDEX_NONE ? true : false :
+		BTParentNode && BTParentNode->Decorators.Find(BTGraphNode) == INDEX_NONE ? true :
 		(BTGraphNode->NodeInstance != NULL && (Cast<UBTNode>(BTGraphNode->NodeInstance->GetOuter()) == NULL && Cast<UBehaviorTree>(BTGraphNode->NodeInstance->GetOuter())==NULL)) ? true : false;
 
 	if (FBehaviorTreeDebugger::IsPIENotSimulating() && BTGraphNode)
@@ -340,6 +277,7 @@ void SGraphNode_BehaviorTree::UpdateGraphNode()
 	LeftNodeBox.Reset();
 	DecoratorWidgets.Reset();
 	ServicesWidgets.Reset();
+	SubNodes.Reset();
 	OutputPinBox.Reset();
 
 	UBehaviorTreeGraphNode* BTNode = Cast<UBehaviorTreeGraphNode>(GraphNode);
@@ -673,251 +611,11 @@ void SGraphNode_BehaviorTree::Tick( const FGeometry& AllottedGeometry, const dou
 	FlashAlpha = NewFlashAlpha;
 }
 
-FReply SGraphNode_BehaviorTree::OnMouseMove(const FGeometry& SenderGeometry, const FPointerEvent& MouseEvent)
-{
-	if (MouseEvent.IsMouseButtonDown(EKeys::LeftMouseButton) && FBehaviorTreeDebugger::IsPIENotSimulating())
-	{
-		//if we are holding mouse over a subnode
-		if (Cast<UBehaviorTreeGraphNode>(GraphNode)->ParentNode)
-		{
-			const TSharedRef<SGraphPanel>& Panel = this->GetOwnerPanel().ToSharedRef();
-			const TSharedRef<SGraphNode>& Node = SharedThis(this);
-			return FReply::Handled().BeginDragDrop(FDragNodeTimed::New(Panel, Node));
-		}
-	}
-
-	if (!MouseEvent.IsMouseButtonDown(EKeys::LeftMouseButton) && bDragMarkerVisible)
-	{
-		SetDragMarker(false);
-	}
-
-	return FReply::Unhandled();
-}
-
-TSharedPtr<SGraphNode> SGraphNode_BehaviorTree::GetSubNodeUnderCursor(const FGeometry& WidgetGeometry, const FPointerEvent& MouseEvent) const
-{
-	TSharedPtr<SGraphNode> ResultNode;
-
-	// We just need to find the one WidgetToFind among our descendants.
-	TSet< TSharedRef<SWidget> > SubWidgetsSet;
-
-	for (int32 i=0; i < DecoratorWidgets.Num(); i++)
-	{
-		SubWidgetsSet.Add(DecoratorWidgets[i].ToSharedRef());
-	}
-	for (int32 i=0; i < ServicesWidgets.Num(); i++)
-	{
-		SubWidgetsSet.Add(ServicesWidgets[i].ToSharedRef());
-	}
-	TMap<TSharedRef<SWidget>, FArrangedWidget> Result;
-
-	FindChildGeometries(WidgetGeometry, SubWidgetsSet, Result);
-
-	if ( Result.Num() > 0 )
-	{
-		FArrangedChildren ArrangedChildren(EVisibility::Visible);
-		Result.GenerateValueArray( ArrangedChildren.GetInternalArray() );
-		int32 HoveredIndex = SWidget::FindChildUnderMouse( ArrangedChildren, MouseEvent );
-		if ( HoveredIndex != INDEX_NONE )
-		{
-			ResultNode = StaticCastSharedRef<SGraphNode>(ArrangedChildren[HoveredIndex].Widget);
-		}
-	}
-
-	return ResultNode;
-}
-
-FReply SGraphNode_BehaviorTree::OnMouseDown(const FGeometry& SenderGeometry, const FPointerEvent& MouseEvent)
-{
-	if (Cast<UBehaviorTreeGraphNode>(GraphNode)->ParentNode)
-	{
-		GetOwnerPanel()->SelectionManager.ClickedOnNode(GraphNode,MouseEvent);
-		return FReply::Handled();
-	}
-
-	return FReply::Unhandled();
-}
-
 FReply SGraphNode_BehaviorTree::OnMouseButtonDoubleClick( const FGeometry& InMyGeometry, const FPointerEvent& InMouseEvent )
 {
 	return SGraphNode::OnMouseButtonDoubleClick(InMyGeometry, InMouseEvent );
 }
 
-void SGraphNode_BehaviorTree::SetDragMarker(bool bEnabled)
-{
-	bDragMarkerVisible = bEnabled;
-}
-
-EVisibility SGraphNode_BehaviorTree::GetDragOverMarkerVisibility() const
-{
-	return bDragMarkerVisible ? EVisibility::Visible : EVisibility::Collapsed; 
-}
-
-void SGraphNode_BehaviorTree::OnDragEnter( const FGeometry& MyGeometry, const FDragDropEvent& DragDropEvent )
-{
-	// Is someone dragging a node?
-	TSharedPtr<FDragNode> DragConnectionOp = DragDropEvent.GetOperationAs<FDragNode>();
-	if (DragConnectionOp.IsValid())
-	{
-		// Inform the Drag and Drop operation that we are hovering over this node.
-		TSharedPtr<SGraphNode> SubNode = GetSubNodeUnderCursor(MyGeometry, DragDropEvent);
-		DragConnectionOp->SetHoveredNode( SubNode.IsValid() ? SubNode : SharedThis(this) );
-		if (DragConnectionOp->IsValidOperation() && Cast<UBehaviorTreeGraphNode>(GraphNode)->ParentNode)
-		{
-			SetDragMarker(true);
-		}
-	}
-	
-	SGraphNode::OnDragEnter(MyGeometry, DragDropEvent);
-}
-
-FReply SGraphNode_BehaviorTree::OnDragOver( const FGeometry& MyGeometry, const FDragDropEvent& DragDropEvent )
-{
-	// Is someone dragging a node?
-	TSharedPtr<FDragNode> DragConnectionOp = DragDropEvent.GetOperationAs<FDragNode>();
-	if (DragConnectionOp.IsValid())
-	{
-		// Inform the Drag and Drop operation that we are hovering over this node.
-		TSharedPtr<SGraphNode> SubNode = GetSubNodeUnderCursor(MyGeometry, DragDropEvent);
-		DragConnectionOp->SetHoveredNode( SubNode.IsValid() ? SubNode : SharedThis(this) );
-	}
-	return SGraphNode::OnDragOver(MyGeometry, DragDropEvent);
-}
-
-void SGraphNode_BehaviorTree::OnDragLeave( const FDragDropEvent& DragDropEvent )
-{
-	TSharedPtr<FDragNode> DragConnectionOp = DragDropEvent.GetOperationAs<FDragNode>();
-	if (DragConnectionOp.IsValid())
-	{
-		// Inform the Drag and Drop operation that we are not hovering any pins
-		DragConnectionOp->SetHoveredNode( TSharedPtr<SGraphNode>(NULL) );
-	}
-
-	SetDragMarker(false);
-	SGraphNode::OnDragLeave(DragDropEvent);
-}
-
-FReply SGraphNode_BehaviorTree::OnDrop( const FGeometry& MyGeometry, const FDragDropEvent& DragDropEvent )
-{
-	SetDragMarker(false);
-
-	TSharedPtr<FDragNodeTimed> DragNodeOp = DragDropEvent.GetOperationAs<FDragNodeTimed>();
-	if (DragNodeOp.IsValid())
-	{
-		if (!DragNodeOp->IsValidOperation())
-		{
-			return FReply::Handled();
-		}
-
-		const float DragTime = float(FPlatformTime::Seconds() - DragNodeOp->StartTime);
-		if (DragTime < 0.5f)
-		{
-			return FReply::Handled();
-		}
-
-		//dropped on decorator, route to parent
-		if (Cast<UBehaviorTreeGraphNode>(GraphNode)->ParentNode != NULL)
-		{
-			return FReply::Unhandled();
-		}
-
-		const FScopedTransaction Transaction( NSLOCTEXT("UnrealEd", "GraphEd_DragDropNode", "Drag&Drop Node") );
-
-		const TArray< TSharedRef<SGraphNode> >& DraggedNodes = DragNodeOp->GetNodes();
-		bool bReorderOperation = true;
-
-		for (int32 Idx = 0; Idx < DraggedNodes.Num(); Idx++)
-		{
-			bool bIsDraggedNodeService = Cast<UBehaviorTreeGraphNode_Service>(DraggedNodes[Idx]->GetNodeObj()) != NULL;
-
-			UBehaviorTreeGraphNode* ParentNode = (Cast<UBehaviorTreeGraphNode>(DraggedNodes[Idx]->GetNodeObj()))->ParentNode;
-			if (bReorderOperation && ParentNode != Cast<UBehaviorTreeGraphNode>(GraphNode))
-			{
-				bReorderOperation = false;
-			}
-
-			if (ParentNode)
-			{
-				ParentNode->Modify();
-				TArray<UBehaviorTreeGraphNode*> & SubNodes = bIsDraggedNodeService ? ParentNode->Services : ParentNode->Decorators;
-				for (int32 SubIdx = 0; SubIdx < SubNodes.Num(); SubIdx++)
-				{
-					if (SubNodes[SubIdx] == DraggedNodes[Idx]->GetNodeObj())
-					{
-						SubNodes.RemoveAt(SubIdx);
-					}
-				}
-			}
-		}
-
-		UBehaviorTreeGraphNode* BTNode = Cast<UBehaviorTreeGraphNode>(GraphNode);
-		UBehaviorTreeGraphNode* DropTargetNode = DragNodeOp->GetDropTargetNode();
-		int32 InsertIndex = -1;
-
-		{
-			const int32 DecoratorIdx = BTNode->Decorators.IndexOfByKey(DropTargetNode);
-			if (DecoratorIdx >= 0)
-			{
-				InsertIndex = DecoratorIdx;
-			}
-			else
-			{
-				const int32 ServiceIdx = BTNode->Services.IndexOfByKey(DropTargetNode);
-				if (ServiceIdx >= 0)
-				{
-					InsertIndex = ServiceIdx;
-				}
-			}
-		}
-
-		for (int32 Idx = 0; Idx < DraggedNodes.Num(); Idx++)
-		{
-			UBehaviorTreeGraphNode* DraggedNode = Cast<UBehaviorTreeGraphNode>(DraggedNodes[Idx]->GetNodeObj());
-			bool bIsDraggedNodeService = Cast<UBehaviorTreeGraphNode_Service>(DraggedNode) != NULL;
-			TArray<UBehaviorTreeGraphNode*> & SubNodes = bIsDraggedNodeService ? BTNode->Services : BTNode->Decorators;
-
-			DraggedNode->Modify();
-			DraggedNode->ParentNode = BTNode;
-			BTNode->Modify();
-
-			if (InsertIndex > -1)
-			{
-				SubNodes.Insert(DraggedNode, InsertIndex);
-			}
-			else
-			{
-				SubNodes.Add(DraggedNode);
-			}
-		}
-
-		if (bReorderOperation) // if updating this node is enough, do it
-		{
-			UpdateGraphNode();
-		}
-		else // if decorator was dragged between nodes, rebuild graph
-		{
-			GraphNode->GetGraph()->NotifyGraphChanged();
-
-			UBehaviorTreeGraph* MyGraph = Cast<UBehaviorTreeGraph>(GraphNode->GetGraph());
-			if (MyGraph)
-			{
-				FAbortDrawHelper EmptyMode;
-
-				MyGraph->UpdateAsset(UBehaviorTreeGraph::ClearDebuggerFlags);
-				MyGraph->UpdateAbortHighlight(EmptyMode, EmptyMode);
-			}
-		}
-	}
-
-	return SGraphNode::OnDrop(MyGeometry, DragDropEvent);
-}
-
-FText SGraphNode_BehaviorTree::GetDescription() const
-{
-	UBehaviorTreeGraphNode* StateNode = CastChecked<UBehaviorTreeGraphNode>(GraphNode);
-	return FText::FromString(StateNode->GetDescription());
-}
-
 FText SGraphNode_BehaviorTree::GetPinTooltip(UEdGraphPin* GraphPinObj) const
 {
 	FText HoverText = FText::GetEmpty();
@@ -937,7 +635,6 @@ FText SGraphNode_BehaviorTree::GetPinTooltip(UEdGraphPin* GraphPinObj) const
 	return HoverText;
 }
 
-
 void SGraphNode_BehaviorTree::CreatePinWidgets()
 {
 	UBehaviorTreeGraphNode* StateNode = CastChecked<UBehaviorTreeGraphNode>(GraphNode);
@@ -1073,12 +770,13 @@ const FSlateBrush* SGraphNode_BehaviorTree::GetNameIcon() const
 
 static UBehaviorTreeGraphNode* GetParentNode(UEdGraphNode* GraphNode)
 {
-	UBehaviorTreeGraphNode* StateNode = CastChecked<UBehaviorTreeGraphNode>(GraphNode);
-	if(StateNode->ParentNode != nullptr)
+	UBehaviorTreeGraphNode* BTGraphNode = Cast<UBehaviorTreeGraphNode>(GraphNode);
+	if (BTGraphNode->ParentNode != nullptr)
 	{
-		StateNode = StateNode->ParentNode;
+		BTGraphNode = Cast<UBehaviorTreeGraphNode>(BTGraphNode->ParentNode);
 	}
-	UEdGraphPin* MyInputPin = StateNode->GetInputPin();
+
+	UEdGraphPin* MyInputPin = BTGraphNode->GetInputPin();
 	UEdGraphPin* MyParentOutputPin = nullptr;
 	if (MyInputPin != nullptr && MyInputPin->LinkedTo.Num() > 0)
 	{
@@ -1310,27 +1008,4 @@ TSharedRef<SGraphNode> SGraphNode_BehaviorTree::GetNodeUnderMouse(const FGeometr
 	return SubNode.IsValid() ? SubNode.ToSharedRef() : StaticCastSharedRef<SGraphNode>(AsShared());
 }
 
-void SGraphNode_BehaviorTree::SetOwner( const TSharedRef<SGraphPanel>& OwnerPanel )
-{
-	SGraphNode::SetOwner(OwnerPanel);
-
-	for(auto& DecoratorWidget : DecoratorWidgets)
-	{
-		if(DecoratorWidget.IsValid())
-		{
-			DecoratorWidget->SetOwner(OwnerPanel);
-			OwnerPanel->AttachGraphEvents(DecoratorWidget);
-		}
-	}
-
-	for(auto& ServicesWidget : ServicesWidgets)
-	{
-		if(ServicesWidget.IsValid())
-		{
-			ServicesWidget->SetOwner(OwnerPanel);
-			OwnerPanel->AttachGraphEvents(ServicesWidget);
-		}
-	}
-}
-
 #undef LOCTEXT_NAMESPACE
\ No newline at end of file
diff --git a/Engine/Source/Editor/BehaviorTreeEditor/Private/SGraphNode_BehaviorTree.h b/Engine/Source/Editor/BehaviorTreeEditor/Private/SGraphNode_BehaviorTree.h
--- a/Engine/Source/Editor/BehaviorTreeEditor/Private/SGraphNode_BehaviorTree.h
+++ b/Engine/Source/Editor/BehaviorTreeEditor/Private/SGraphNode_BehaviorTree.h
@@ -2,25 +2,9 @@
 
 #pragma once
 
-class UBehaviorTreeGraphNode;
+#include "SGraphNodeAI.h"
 
-class FDragNodeTimed : public FDragNode
-{
-public:
-	DRAG_DROP_OPERATOR_TYPE(FDragNodeTimed, FDragNode)
-
-	static TSharedRef<FDragNodeTimed> New(const TSharedRef<SGraphPanel>& InGraphPanel, const TSharedRef<SGraphNode>& InDraggedNode);
-	static TSharedRef<FDragNodeTimed> New(const TSharedRef<SGraphPanel>& InGraphPanel, const TArray< TSharedRef<SGraphNode> >& InDraggedNodes);
-
-	UBehaviorTreeGraphNode* GetDropTargetNode() const;
-
-	double StartTime;
-
-protected:
-	typedef FDragNode Super;
-};
-
-class SGraphNode_BehaviorTree : public SGraphNode
+class SGraphNode_BehaviorTree : public SGraphNodeAI
 {
 public:
 	SLATE_BEGIN_ARGS(SGraphNode_BehaviorTree){}
@@ -33,25 +17,14 @@ public:
 	virtual void CreatePinWidgets() override;
 	virtual void AddPin(const TSharedRef<SGraphPin>& PinToAdd) override;
 	virtual TSharedPtr<SToolTip> GetComplexTooltip() override;
-	virtual void OnDragEnter( const FGeometry& MyGeometry, const FDragDropEvent& DragDropEvent ) override;
-	virtual FReply OnDragOver( const FGeometry& MyGeometry, const FDragDropEvent& DragDropEvent ) override;
-	virtual FReply OnDrop( const FGeometry& MyGeometry, const FDragDropEvent& DragDropEvent ) override;
-	virtual void OnDragLeave( const FDragDropEvent& DragDropEvent ) override;
 	virtual void GetOverlayBrushes(bool bSelected, const FVector2D WidgetSize, TArray<FOverlayBrushInfo>& Brushes) const override;
 	virtual TArray<FOverlayWidgetInfo> GetOverlayWidgets(bool bSelected, const FVector2D& WidgetSize) const override;
 	virtual TSharedRef<SGraphNode> GetNodeUnderMouse(const FGeometry& MyGeometry, const FPointerEvent& MouseEvent) override;
-	virtual void SetOwner( const TSharedRef<SGraphPanel>& OwnerPanel ) override;
 	// End of SGraphNode interface
 
 	/** handle double click */
 	virtual FReply OnMouseButtonDoubleClick( const FGeometry& InMyGeometry, const FPointerEvent& InMouseEvent ) override;
 
-	/** handle mouse down on the node */
-	FReply OnMouseDown(const FGeometry& SenderGeometry, const FPointerEvent& MouseEvent);
-
-	/** handle mouse move on the node - activate drag and drop when mouse button is being held */
-	virtual FReply OnMouseMove(const FGeometry& SenderGeometry, const FPointerEvent& MouseEvent) override;
-
 	/**
 	 * Ticks this widget.  Override in derived classes, but always call the parent implementation.
 	 *
@@ -61,29 +34,18 @@ public:
 	 */
 	virtual void Tick( const FGeometry& AllottedGeometry, const double InCurrentTime, const float InDeltaTime ) override;
 	
-	/** gets decorator or service node if one is found under mouse cursor */
-	TSharedPtr<SGraphNode> GetSubNodeUnderCursor(const FGeometry& WidgetGeometry, const FPointerEvent& MouseEvent) const;
-
 	/** adds decorator widget inside current node */
 	void AddDecorator(TSharedPtr<SGraphNode> DecoratorWidget);
 
 	/** adds service widget inside current node */
 	void AddService(TSharedPtr<SGraphNode> ServiceWidget);
 
-	/** gets drag over marker visibility */
-	EVisibility GetDragOverMarkerVisibility() const;
-
 	/** shows red marker when search failed*/
 	EVisibility GetDebuggerSearchFailedMarkerVisibility() const;
 
-	/** sets drag marker visible or collapsed on this node */
-	void SetDragMarker(bool bEnabled);
-
 	FVector2D GetCachedPosition() const { return CachedPosition; }
 
 protected:
-	uint32 bDragMarkerVisible : 1;
-
 	uint32 bSuppressDebuggerColor : 1;
 	uint32 bSuppressDebuggerTriggers : 1;
 	
@@ -117,7 +79,6 @@ protected:
 
 	FSlateColor GetBorderBackgroundColor() const;
 	FSlateColor GetBackgroundColor() const;
-	FText GetDescription() const;
 
 	virtual const FSlateBrush* GetNameIcon() const;
 	virtual EVisibility GetBlueprintIconVisibility() const;
diff --git a/Engine/Source/Editor/BehaviorTreeEditor/Public/BehaviorTreeEditorModule.h b/Engine/Source/Editor/BehaviorTreeEditor/Public/BehaviorTreeEditorModule.h
--- a/Engine/Source/Editor/BehaviorTreeEditor/Public/BehaviorTreeEditorModule.h
+++ b/Engine/Source/Editor/BehaviorTreeEditor/Public/BehaviorTreeEditorModule.h
@@ -26,7 +26,7 @@ public:
 	virtual TSharedPtr<FExtensibilityManager> GetMenuExtensibilityManager() {return MenuExtensibilityManager;}
 	virtual TSharedPtr<FExtensibilityManager> GetToolBarExtensibilityManager() {return ToolBarExtensibilityManager;}
 
-	TSharedPtr<struct FClassBrowseHelper> GetClassCache() { return ClassCache; }
+	TSharedPtr<struct FGraphNodeClassHelper> GetClassCache() { return ClassCache; }
 
 	/** Behavior Tree app identifier string */
 	static const FName BehaviorTreeEditorAppIdentifier;
@@ -38,7 +38,7 @@ private:
 	/** Asset type actions */
 	TArray<TSharedPtr<class FAssetTypeActions_Base>> ItemDataAssetTypeActions;
 
-	TSharedPtr<struct FClassBrowseHelper> ClassCache;
+	TSharedPtr<struct FGraphNodeClassHelper> ClassCache;
 };
 
 
