fixed creating connections from task to composite node in behavior tree editor
#ue4 UE-6594

diff --git a/Engine/Source/Editor/BehaviorTreeEditor/Private/EdGraphSchema_BehaviorTree.cpp b/Engine/Source/Editor/BehaviorTreeEditor/Private/EdGraphSchema_BehaviorTree.cpp
--- a/Engine/Source/Editor/BehaviorTreeEditor/Private/EdGraphSchema_BehaviorTree.cpp
+++ b/Engine/Source/Editor/BehaviorTreeEditor/Private/EdGraphSchema_BehaviorTree.cpp
@@ -215,10 +215,6 @@ const FPinConnectionResponse UEdGraphSchema_BehaviorTree::CanCreateConnection(co
 	{
 		return FPinConnectionResponse(CONNECT_RESPONSE_DISALLOW, LOCTEXT("PinErrorOutput", "Can't connect output node to output node"));
 	}
-	else if ((PinA->Direction == EGPD_Input) && (PinB->Direction == EGPD_Output))
-	{
-		return FPinConnectionResponse(CONNECT_RESPONSE_DISALLOW, LOCTEXT("PinErrorWrongDirection", "Can't connect input node to output node"));
-	}
 
 	class FNodeVisitorCycleChecker
 	{
@@ -282,7 +278,7 @@ const FPinConnectionResponse UEdGraphSchema_BehaviorTree::CanCreateConnection(co
 	const bool bPinASingleLink = bPinAIsSingleComposite || bPinAIsSingleTask || bPinAIsSingleNode;
 	const bool bPinBSingleLink = bPinBIsSingleComposite || bPinBIsSingleTask || bPinBIsSingleNode;
 
-	if (PinB->LinkedTo.Num() > 0)
+	if (PinB->Direction == EGPD_Input && PinB->LinkedTo.Num() > 0)
 	{
 		if(bPinASingleLink)
 		{
@@ -293,6 +289,17 @@ const FPinConnectionResponse UEdGraphSchema_BehaviorTree::CanCreateConnection(co
 			return FPinConnectionResponse(CONNECT_RESPONSE_BREAK_OTHERS_B, LOCTEXT("PinConnectReplace", "Replace connection"));
 		}
 	}
+	else if (PinA->Direction == EGPD_Input && PinA->LinkedTo.Num() > 0)
+	{
+		if (bPinBSingleLink)
+		{
+			return FPinConnectionResponse(CONNECT_RESPONSE_BREAK_OTHERS_AB, LOCTEXT("PinConnectReplace", "Replace connection"));
+		}
+		else
+		{
+			return FPinConnectionResponse(CONNECT_RESPONSE_BREAK_OTHERS_A, LOCTEXT("PinConnectReplace", "Replace connection"));
+		}
+	}
 
 	if (bPinASingleLink && PinA->LinkedTo.Num() > 0)
 	{
