added spawning new wait task when parallel's background branch is empty on behavior tree save
#ue4 UE-11053

diff --git a/Engine/Source/Editor/BehaviorTreeEditor/Classes/BehaviorTreeGraph.h b/Engine/Source/Editor/BehaviorTreeEditor/Classes/BehaviorTreeGraph.h
--- a/Engine/Source/Editor/BehaviorTreeEditor/Classes/BehaviorTreeGraph.h
+++ b/Engine/Source/Editor/BehaviorTreeEditor/Classes/BehaviorTreeGraph.h
@@ -27,6 +27,7 @@ class UBehaviorTreeGraph : public UAIGraph
 	virtual void OnCreated() override;
 	virtual void OnLoaded() override;
 	virtual void Initialize() override;
+	void OnSave();
 
 	virtual void UpdateVersion() override;
 	virtual void MarkVersion() override;
@@ -43,6 +44,7 @@ class UBehaviorTreeGraph : public UAIGraph
 	class UEdGraphNode* FindInjectedNode(int32 Index);
 	void ReplaceNodeConnections(UEdGraphNode* OldNode, UEdGraphNode* NewNode);
 	void RebuildExecutionOrder();
+	void SpawnMissingNodesForParallel();
 
 	void AutoArrange();
 
diff --git a/Engine/Source/Editor/BehaviorTreeEditor/Private/BehaviorTreeEditor.cpp b/Engine/Source/Editor/BehaviorTreeEditor/Private/BehaviorTreeEditor.cpp
--- a/Engine/Source/Editor/BehaviorTreeEditor/Private/BehaviorTreeEditor.cpp
+++ b/Engine/Source/Editor/BehaviorTreeEditor/Private/BehaviorTreeEditor.cpp
@@ -1345,7 +1345,7 @@ void FBehaviorTreeEditor::SaveAsset_Execute()
 		UBehaviorTreeGraph* BTGraph = Cast<UBehaviorTreeGraph>(BehaviorTree->BTGraph);
 		if (BTGraph)
 		{
-			BTGraph->UpdateAsset();
+			BTGraph->OnSave();
 		}
 	}
 	// save it
diff --git a/Engine/Source/Editor/BehaviorTreeEditor/Private/BehaviorTreeGraph.cpp b/Engine/Source/Editor/BehaviorTreeEditor/Private/BehaviorTreeGraph.cpp
--- a/Engine/Source/Editor/BehaviorTreeEditor/Private/BehaviorTreeGraph.cpp
+++ b/Engine/Source/Editor/BehaviorTreeEditor/Private/BehaviorTreeGraph.cpp
@@ -7,6 +7,7 @@
 #include "BehaviorTree/BTService.h"
 #include "BehaviorTree/BehaviorTree.h"
 #include "BehaviorTree/Composites/BTComposite_SimpleParallel.h"
+#include "BehaviorTree/Tasks/BTTask_Wait.h"
 
 //////////////////////////////////////////////////////////////////////////
 // BehaviorTreeGraph
@@ -168,6 +169,12 @@ void UBehaviorTreeGraph::Initialize()
 	UpdateInjectedNodes();
 }
 
+void UBehaviorTreeGraph::OnSave()
+{
+	SpawnMissingNodesForParallel();
+	UpdateAsset();
+}
+
 void UBehaviorTreeGraph::UpdatePinConnectionTypes()
 {
 	for (int32 Index = 0; Index < Nodes.Num(); ++Index)
@@ -873,6 +880,67 @@ void UBehaviorTreeGraph::SpawnMissingNodes()
 	}
 }
 
+void UBehaviorTreeGraph::SpawnMissingNodesForParallel()
+{
+	UBehaviorTree* BTAsset = Cast<UBehaviorTree>(GetOuter());
+	if (BTAsset)
+	{
+		TArray<UBehaviorTreeGraphNode_SimpleParallel*> FixNodes;
+		for (int32 Idx = 0; Idx < Nodes.Num(); Idx++)
+		{
+			UBehaviorTreeGraphNode_SimpleParallel* ParallelNode = Cast<UBehaviorTreeGraphNode_SimpleParallel>(Nodes[Idx]);
+			if (ParallelNode)
+			{
+				UEdGraphPin* BackgroundPin = ParallelNode->GetOutputPin(1);
+				if (BackgroundPin && BackgroundPin->LinkedTo.Num() == 0)
+				{
+					FixNodes.Add(ParallelNode);
+				}
+			}
+		}
+
+		for (int32 Idx = 0; Idx < FixNodes.Num(); Idx++)
+		{
+			UBehaviorTreeGraphNode_SimpleParallel* ParallelNode = FixNodes[Idx];
+			UBTComposite_SimpleParallel* ParallelInstance = Cast<UBTComposite_SimpleParallel>(ParallelNode->NodeInstance);
+			if (ParallelInstance)
+			{
+				int32 XOffset = 200;
+
+				UEdGraphPin* MainTaskPin = ParallelNode->GetOutputPin(0);
+				if (MainTaskPin && MainTaskPin->LinkedTo.Num())
+				{
+					UBehaviorTreeGraphNode* MainTaskNode = Cast<UBehaviorTreeGraphNode>(MainTaskPin->LinkedTo[0]->GetOwningNode());
+					if (MainTaskNode)
+					{
+						const int32 Width = MainTaskNode->NodeWidget.IsValid() ? MainTaskNode->NodeWidget.Pin()->GetDesiredSize().X : 200;
+						XOffset = MainTaskNode->NodePosX - ParallelNode->NodePosX + Width + 20;
+					}
+				}
+
+				FGraphNodeCreator<UBehaviorTreeGraphNode_Task> NodeBuilder(*this);
+				UBehaviorTreeGraphNode_Task* WaitTaskNode = NodeBuilder.CreateNode();
+				WaitTaskNode->ClassData = FGraphNodeClassData(UBTTask_Wait::StaticClass(), "");
+				NodeBuilder.Finalize();
+
+				const int32 ParentHeight = ParallelNode->NodeWidget.IsValid() ? ParallelNode->NodeWidget.Pin()->GetDesiredSize().Y : 200;
+				WaitTaskNode->NodePosX = ParallelNode->NodePosX + XOffset;
+				WaitTaskNode->NodePosY = ParallelNode->NodePosY + ParentHeight + 20;
+
+				UBTTask_Wait* WaitTaskInstance = Cast<UBTTask_Wait>(WaitTaskNode->NodeInstance);
+				if (WaitTaskInstance)
+				{
+					WaitTaskInstance->WaitTime = (ParallelInstance->FinishMode == EBTParallelMode::WaitForBackground) ? 0.5f : 60.0f;
+				}
+
+				UEdGraphPin* BackgroundPin = ParallelNode->GetOutputPin(1);
+				UEdGraphPin* InputPin = WaitTaskNode->GetInputPin();
+				BackgroundPin->MakeLinkTo(InputPin);
+			}
+		}
+	}
+}
+
 void UBehaviorTreeGraph::UpdateAbortHighlight(struct FAbortDrawHelper& Mode0, struct FAbortDrawHelper& Mode1)
 {
 	for (int32 Index = 0; Index < Nodes.Num(); ++Index)
