Renaming blackboard keys now automaticaly updates Behavior Trees that reference them

Note that to do this we need to load all behavior trees in the asset registry to check for blackboard/key usage

UE-8415 - Renaming a Blackboard Key removes it's references in the Behavior Tree

diff --git a/Engine/Source/Editor/BehaviorTreeEditor/Private/SBehaviorTreeBlackboardEditor.cpp b/Engine/Source/Editor/BehaviorTreeEditor/Private/SBehaviorTreeBlackboardEditor.cpp
--- a/Engine/Source/Editor/BehaviorTreeEditor/Private/SBehaviorTreeBlackboardEditor.cpp
+++ b/Engine/Source/Editor/BehaviorTreeEditor/Private/SBehaviorTreeBlackboardEditor.cpp
@@ -13,6 +13,7 @@
 #include "ClassViewerModule.h"
 #include "ClassViewerFilter.h"
 #include "GenericCommands.h"
+#include "SBehaviorTreeBlackboardView.h"
 
 #define LOCTEXT_NAMESPACE "SBehaviorTreeBlackboardEditor"
 
@@ -259,6 +260,15 @@ void SBehaviorTreeBlackboardEditor::HandleKeyClassPicked(UClass* InClass)
 	OnBlackboardKeyChanged.ExecuteIfBound(BlackboardData, &BlackboardData->Keys.Last());
 
 	GraphActionMenu->SelectItemByName(Entry.EntryName, ESelectInfo::OnMouseClick);
+
+	// Mark newly created entry as 'new'
+	TArray< TSharedPtr<FEdGraphSchemaAction> > SelectedActions;
+	GraphActionMenu->GetSelectedActions(SelectedActions);
+	check(SelectedActions.Num() == 1);
+	check(SelectedActions[0]->GetTypeId() == FEdGraphSchemaAction_BlackboardEntry::StaticGetTypeId());
+	TSharedPtr<FEdGraphSchemaAction_BlackboardEntry> BlackboardEntryAction = StaticCastSharedPtr<FEdGraphSchemaAction_BlackboardEntry>(SelectedActions[0]);
+	BlackboardEntryAction->bIsNew = true;
+
 	GraphActionMenu->OnRequestRenameOnActionNode();
 }
 
diff --git a/Engine/Source/Editor/BehaviorTreeEditor/Private/SBehaviorTreeBlackboardView.cpp b/Engine/Source/Editor/BehaviorTreeEditor/Private/SBehaviorTreeBlackboardView.cpp
--- a/Engine/Source/Editor/BehaviorTreeEditor/Private/SBehaviorTreeBlackboardView.cpp
+++ b/Engine/Source/Editor/BehaviorTreeEditor/Private/SBehaviorTreeBlackboardView.cpp
@@ -2,6 +2,7 @@
 
 #include "BehaviorTreeEditorPrivatePCH.h"
 #include "SBehaviorTreeBlackboardView.h"
+#include "BehaviorTree/BehaviorTree.h"
 #include "BehaviorTree/BlackboardData.h"
 #include "BehaviorTree/Blackboard/BlackboardKeyAllTypes.h"
 #include "SGraphActionMenu.h"
@@ -13,6 +14,7 @@
 #include "ClassViewerModule.h"
 #include "ClassViewerFilter.h"
 #include "SInlineEditableTextBlock.h"
+#include "AssetRegistryModule.h"
 
 #define LOCTEXT_NAMESPACE "SBehaviorTreeBlackboardView"
 
@@ -25,37 +27,33 @@ namespace EBlackboardSectionTitles
 	};
 }
 
+FName FEdGraphSchemaAction_BlackboardEntry::StaticGetTypeId() 
+{ 
+	static FName Type("FEdGraphSchemaAction_BlackboardEntry"); return Type; 
+}
 
-class FEdGraphSchemaAction_BlackboardEntry : public FEdGraphSchemaAction_Dummy
-{
-public:
-	static FName StaticGetTypeId() { static FName Type("FEdGraphSchemaAction_BlackboardEntry"); return Type; }
-	virtual FName GetTypeId() const override { return StaticGetTypeId(); }
-
-	FEdGraphSchemaAction_BlackboardEntry( UBlackboardData* InBlackboardData, FBlackboardEntry& InKey, bool bInIsInherited )
-		: FEdGraphSchemaAction_Dummy()
-		, BlackboardData(InBlackboardData)
-		, Key(InKey)
-		, bIsInherited(bInIsInherited)
-	{
-		check(BlackboardData);
-		Update();
-	}
-
-	void Update()
-	{
-		MenuDescription = FText::FromName(Key.EntryName);
-		TooltipDescription = FText::Format(LOCTEXT("BlackboardEntryFormat", "{0} '{1}'"), Key.KeyType ? Key.KeyType->GetClass()->GetDisplayNameText() : LOCTEXT("NullKeyDesc", "None"), FText::FromName(Key.EntryName)).ToString();
-		SectionID = bIsInherited ? EBlackboardSectionTitles::InheritedKeys : EBlackboardSectionTitles::Keys;	
-	}
-
-	UBlackboardData* BlackboardData;
-
-	FBlackboardEntry& Key;
+FName FEdGraphSchemaAction_BlackboardEntry::GetTypeId() const 
+{ 
+	return StaticGetTypeId(); 
+}
 
-	bool bIsInherited;
-};
+FEdGraphSchemaAction_BlackboardEntry::FEdGraphSchemaAction_BlackboardEntry( UBlackboardData* InBlackboardData, FBlackboardEntry& InKey, bool bInIsInherited )
+	: FEdGraphSchemaAction_Dummy()
+	, BlackboardData(InBlackboardData)
+	, Key(InKey)
+	, bIsInherited(bInIsInherited)
+	, bIsNew(false)
+{
+	check(BlackboardData);
+	Update();
+}
 
+void FEdGraphSchemaAction_BlackboardEntry::Update()
+{
+	MenuDescription = FText::FromName(Key.EntryName);
+	TooltipDescription = FText::Format(LOCTEXT("BlackboardEntryFormat", "{0} '{1}'"), Key.KeyType ? Key.KeyType->GetClass()->GetDisplayNameText() : LOCTEXT("NullKeyDesc", "None"), FText::FromName(Key.EntryName)).ToString();
+	SectionID = bIsInherited ? EBlackboardSectionTitles::InheritedKeys : EBlackboardSectionTitles::Keys;	
+}
 
 class SBehaviorTreeBlackboardItem : public SGraphPaletteItem
 {
@@ -181,17 +179,112 @@ private:
 		check(ActionPtr.Pin()->GetTypeId() == FEdGraphSchemaAction_BlackboardEntry::StaticGetTypeId());
 		TSharedPtr<FEdGraphSchemaAction_BlackboardEntry> BlackboardEntryAction = StaticCastSharedPtr<FEdGraphSchemaAction_BlackboardEntry>(ActionPtr.Pin());
 
+		FName OldName = BlackboardEntryAction->Key.EntryName;
 		FName NewName = FName(*NewText.ToString());
-		if(NewName != BlackboardEntryAction->Key.EntryName)
+		if(NewName != OldName)
 		{
+			if(!BlackboardEntryAction->bIsNew)
+			{
+				// Preload behavior trees before we transact otherwise they will add objects to 
+				// the transaction buffer whether we change them or not.
+				// Blueprint regeneration does this in UEdGraphNode::CreatePin.
+				LoadAllBehaviorTrees();
+			}
+
 			const FScopedTransaction Transaction(LOCTEXT("BlackboardEntryRenameTransaction", "Rename Blackboard Entry"));
 			BlackboardEntryAction->BlackboardData->SetFlags(RF_Transactional);
 			BlackboardEntryAction->BlackboardData->Modify();
-			BlackboardEntryAction->Key.EntryName = FName(*NewText.ToString());
+			BlackboardEntryAction->Key.EntryName = NewName;
 
 			BlackboardEntryAction->Update();
 
 			OnBlackboardKeyChanged.ExecuteIfBound(BlackboardEntryAction->BlackboardData, &BlackboardEntryAction->Key);
+
+			if(!BlackboardEntryAction->bIsNew)
+			{
+				UpdateExternalBlackboardKeyReferences(OldName, NewName);
+			}
+		}
+
+		BlackboardEntryAction->bIsNew = false;
+	}
+
+	void LoadAllBehaviorTrees()
+	{
+		FAssetRegistryModule& AssetRegistry = FModuleManager::LoadModuleChecked<FAssetRegistryModule>(TEXT("AssetRegistry"));
+
+		FARFilter Filter;
+		Filter.ClassNames.Add(UBehaviorTree::StaticClass()->GetFName());
+		Filter.bRecursiveClasses = true;
+
+		TArray<FAssetData> AssetData;
+		AssetRegistry.Get().GetAssets(Filter, AssetData);
+
+		FScopedSlowTask SlowTask((float)AssetData.Num(), LOCTEXT("UpdatingBehaviorTrees", "Updating behavior trees"));
+		SlowTask.MakeDialog();
+
+		for (const auto& BehaviorTreeAsset : AssetData)
+		{
+			SlowTask.EnterProgressFrame(1.0f, FText::Format(LOCTEXT("CheckingBehaviorTree", "Key renamed, loading {0}"), FText::FromName(BehaviorTreeAsset.AssetName)));
+			BehaviorTreeAsset.GetAsset();
+		}
+	}
+
+	#define GET_STRUCT_NAME_CHECKED(StructName) \
+		((void)sizeof(StructName), TEXT(#StructName))
+
+	void UpdateExternalBlackboardKeyReferences(const FName& OldKey, const FName& NewKey) const
+	{
+		// update all behavior trees that reference this key
+		check(ActionPtr.Pin()->GetTypeId() == FEdGraphSchemaAction_BlackboardEntry::StaticGetTypeId());
+		TSharedPtr<FEdGraphSchemaAction_BlackboardEntry> BlackboardEntryAction = StaticCastSharedPtr<FEdGraphSchemaAction_BlackboardEntry>(ActionPtr.Pin());
+
+		FAssetRegistryModule& AssetRegistry = FModuleManager::LoadModuleChecked<FAssetRegistryModule>(TEXT("AssetRegistry"));
+
+		FARFilter Filter;
+		Filter.ClassNames.Add(UBehaviorTree::StaticClass()->GetFName());
+		Filter.bRecursiveClasses = true;
+
+		TArray<FAssetData> AssetData;
+		AssetRegistry.Get().GetAssets(Filter, AssetData);
+
+		for (const auto& BehaviorTreeAsset : AssetData)
+		{
+			UBehaviorTree* BehaviorTree = Cast<UBehaviorTree>(BehaviorTreeAsset.GetAsset());
+			if(BehaviorTree)
+			{
+				if(BehaviorTree->BlackboardAsset == BlackboardEntryAction->BlackboardData)
+				{
+					UPackage* BehaviorTreePackage = BehaviorTreeAsset.GetPackage();
+					if(BehaviorTreePackage)
+					{
+						// search all subobjects of this package for FBlackboardKeySelector structs and update as necessary
+						TArray<UObject*> Objects;
+						GetObjectsWithOuter(BehaviorTreePackage, Objects);
+						for(const auto& SubObject : Objects)
+						{
+							for (UProperty* Property = SubObject->GetClass()->PropertyLink; Property; Property = Property->PropertyLinkNext)
+							{
+								uint8* PropertyData = Property->ContainerPtrToValuePtr<uint8>(SubObject);
+								UStructProperty* StructProperty = Cast<UStructProperty>(Property);
+
+								if (StructProperty && StructProperty->GetCPPType(NULL, CPPF_None).Contains(GET_STRUCT_NAME_CHECKED(FBlackboardKeySelector)))
+								{
+									FBlackboardKeySelector* PropertyValue = (FBlackboardKeySelector*)PropertyData;
+									if(PropertyValue)
+									{
+										if(PropertyValue->SelectedKeyName == OldKey)
+										{
+											SubObject->Modify();
+											PropertyValue->SelectedKeyName = NewKey;
+										}
+									}
+								}
+							}
+						}
+					}
+				}
+			}
 		}
 	}
 
diff --git a/Engine/Source/Editor/BehaviorTreeEditor/Private/SBehaviorTreeBlackboardView.h b/Engine/Source/Editor/BehaviorTreeEditor/Private/SBehaviorTreeBlackboardView.h
--- a/Engine/Source/Editor/BehaviorTreeEditor/Private/SBehaviorTreeBlackboardView.h
+++ b/Engine/Source/Editor/BehaviorTreeEditor/Private/SBehaviorTreeBlackboardView.h
@@ -7,7 +7,6 @@
 
 class UBlackboardData;
 struct FBlackboardEntry;
-class FEdGraphSchemaAction_BlackboardEntry;
 
 /** Delegate executed when an entry is selected */
 DECLARE_DELEGATE_TwoParams( FOnEntrySelected, const FBlackboardEntry*, bool );
@@ -30,6 +29,30 @@ DECLARE_DELEGATE_RetVal_OneParam(float, FOnGetDebugTimeStamp, bool /* bUseCurren
 /** Delegate for when a blackboard key changes (added, removed, renamed) */
 DECLARE_DELEGATE_TwoParams(FOnBlackboardKeyChanged, UBlackboardData* /*InBlackboardData*/, FBlackboardEntry* const /*InKey*/);
 
+/** Blackboard entry in the list */
+class FEdGraphSchemaAction_BlackboardEntry : public FEdGraphSchemaAction_Dummy
+{
+public:
+	static FName StaticGetTypeId();
+	virtual FName GetTypeId() const;
+
+	FEdGraphSchemaAction_BlackboardEntry( UBlackboardData* InBlackboardData, FBlackboardEntry& InKey, bool bInIsInherited );
+
+	void Update();
+
+	/** Blackboard we reference our key in */
+	UBlackboardData* BlackboardData;
+
+	/** Actual key */
+	FBlackboardEntry& Key;
+
+	/** Whether this entry came from a parent blackboard */
+	bool bIsInherited;
+
+	/** Temp flag for new items */
+	bool bIsNew;
+};
+
 /** Displays blackboard entries */
 class SBehaviorTreeBlackboardView : public SCompoundWidget, public FGCObject
 {
