fixed drag & drop operation for composite decorators
#ue4, merged from dev branch

diff --git a/Engine/Source/Editor/BehaviorTreeEditor/Private/BehaviorTreeGraphNode.cpp b/Engine/Source/Editor/BehaviorTreeEditor/Private/BehaviorTreeGraphNode.cpp
--- a/Engine/Source/Editor/BehaviorTreeEditor/Private/BehaviorTreeGraphNode.cpp
+++ b/Engine/Source/Editor/BehaviorTreeEditor/Private/BehaviorTreeGraphNode.cpp
@@ -186,7 +186,7 @@ void UBehaviorTreeGraphNode::InsertSubNodeAt(UAIGraphNode* SubNode, int32 DropIn
 	}
 
 	UBehaviorTreeGraphNode* TypedNode = Cast<UBehaviorTreeGraphNode>(SubNode);
-	const bool bIsDecorator = Cast<UBTDecorator>(SubNode->NodeInstance) != nullptr;
+	const bool bIsDecorator = (Cast<UBTDecorator>(SubNode->NodeInstance) != nullptr) || (Cast<UBehaviorTreeGraphNode_CompositeDecorator>(SubNode) != nullptr);
 	const bool bIsService = Cast<UBTService>(SubNode->NodeInstance) != nullptr;
 
 	if (TypedNode)
