Made blackboard key name editable in details tab #UE4

UE-7112

diff --git a/Engine/Source/Editor/BehaviorTreeEditor/Private/DetailCustomizations/BlackboardDataDetails.cpp b/Engine/Source/Editor/BehaviorTreeEditor/Private/DetailCustomizations/BlackboardDataDetails.cpp
--- a/Engine/Source/Editor/BehaviorTreeEditor/Private/DetailCustomizations/BlackboardDataDetails.cpp
+++ b/Engine/Source/Editor/BehaviorTreeEditor/Private/DetailCustomizations/BlackboardDataDetails.cpp
@@ -45,7 +45,7 @@ void FBlackboardDataDetails::CustomizeDetails( IDetailLayoutBuilder& DetailLayou
 			.ValueContent()
 			[
 				SNew(SHorizontalBox)
-				.IsEnabled(false)
+				.IsEnabled(true)
 				+SHorizontalBox::Slot()
 				[
 					EntryNameProperty->CreatePropertyValueWidget()
