fixed duplicated subnodes in behavior tree editor
#ue4 UE-14368

diff --git a/Engine/Source/Editor/AIGraph/Classes/AIGraphNode.h b/Engine/Source/Editor/AIGraph/Classes/AIGraphNode.h
--- a/Engine/Source/Editor/AIGraph/Classes/AIGraphNode.h
+++ b/Engine/Source/Editor/AIGraph/Classes/AIGraphNode.h
@@ -29,12 +29,15 @@ class AIGRAPH_API UAIGraphNode : public UEdGraphNode
 	int32 CopySubNodeIndex;
 
 	/** if set, all modifications (including delete/cut) are disabled */
+	UPROPERTY()
 	uint32 bIsReadOnly : 1;
 
 	/** if set, this node will be always considered as subnode */
+	UPROPERTY()
 	uint32 bIsSubNode : 1;
 
 	/** error message for node */
+	UPROPERTY()
 	FString ErrorMessage;
 
 	// Begin UEdGraphNode Interface
diff --git a/Engine/Source/Editor/BehaviorTreeEditor/Classes/BehaviorTreeGraph.h b/Engine/Source/Editor/BehaviorTreeEditor/Classes/BehaviorTreeGraph.h
--- a/Engine/Source/Editor/BehaviorTreeEditor/Classes/BehaviorTreeGraph.h
+++ b/Engine/Source/Editor/BehaviorTreeEditor/Classes/BehaviorTreeGraph.h
@@ -45,6 +45,7 @@ class UBehaviorTreeGraph : public UAIGraph
 	void ReplaceNodeConnections(UEdGraphNode* OldNode, UEdGraphNode* NewNode);
 	void RebuildExecutionOrder();
 	void SpawnMissingNodesForParallel();
+	void RemoveUnknownSubNodes();
 
 	void AutoArrange();
 
diff --git a/Engine/Source/Editor/BehaviorTreeEditor/Private/BehaviorTreeGraph.cpp b/Engine/Source/Editor/BehaviorTreeEditor/Private/BehaviorTreeGraph.cpp
--- a/Engine/Source/Editor/BehaviorTreeEditor/Private/BehaviorTreeGraph.cpp
+++ b/Engine/Source/Editor/BehaviorTreeEditor/Private/BehaviorTreeGraph.cpp
@@ -161,6 +161,7 @@ void UBehaviorTreeGraph::OnLoaded()
 	Super::OnLoaded();
 	UpdatePinConnectionTypes();
 	UpdateDeprecatedNodes();
+	RemoveUnknownSubNodes();
 }
 
 void UBehaviorTreeGraph::Initialize()
@@ -240,6 +241,27 @@ void UBehaviorTreeGraph::UpdateDeprecatedNodes()
 	}
 }
 
+void UBehaviorTreeGraph::RemoveUnknownSubNodes()
+{
+	for (int32 Index = 0; Index < Nodes.Num(); ++Index)
+	{
+		UBehaviorTreeGraphNode* Node = Cast<UBehaviorTreeGraphNode>(Nodes[Index]);
+		if (Node)
+		{
+			for (int32 SubIdx = Node->SubNodes.Num() - 1; SubIdx >= 0; SubIdx--)
+			{
+				const bool bIsDecorator = Node->Decorators.Contains(Node->SubNodes[SubIdx]);
+				const bool bIsService = Node->Services.Contains(Node->SubNodes[SubIdx]);
+
+				if (!bIsDecorator && !bIsService)
+				{
+					Node->SubNodes.RemoveAt(SubIdx);
+				}
+			}
+		}
+	}
+}
+
 namespace BTGraphHelpers
 {
 	struct FIntIntPair
diff --git a/Engine/Source/Editor/BehaviorTreeEditor/Private/BehaviorTreeGraphNode_SubtreeTask.cpp b/Engine/Source/Editor/BehaviorTreeEditor/Private/BehaviorTreeGraphNode_SubtreeTask.cpp
--- a/Engine/Source/Editor/BehaviorTreeEditor/Private/BehaviorTreeGraphNode_SubtreeTask.cpp
+++ b/Engine/Source/Editor/BehaviorTreeEditor/Private/BehaviorTreeGraphNode_SubtreeTask.cpp
@@ -37,6 +37,7 @@ bool UBehaviorTreeGraphNode_SubtreeTask::UpdateInjectedNodes()
 	{
 		if (Decorators[Index] && Decorators[Index]->bInjectedNode)
 		{
+			SubNodes.RemoveSingle(Decorators[Index]);
 			Decorators.RemoveAt(Index, 1, false);
 		}
 	}
