Some workflow and performance improvements for behavior trees and environmental queries

UETOOL-293 - Behavior Tree Editor is sluggish

- Made the behavior tree and environmental query pins larger so they're easier to grab.
- Adjusted the node padding so that we don't add the padding for the pins when there aren't any pins.
- Added lodding so that node elements are culled out as you zoom out (to improve performance).
- Fixed connections being culled when they were still visible.

diff --git a/Engine/Source/Editor/AIGraph/Private/SGraphNodeAI.cpp b/Engine/Source/Editor/AIGraph/Private/SGraphNodeAI.cpp
--- a/Engine/Source/Editor/AIGraph/Private/SGraphNodeAI.cpp
+++ b/Engine/Source/Editor/AIGraph/Private/SGraphNodeAI.cpp
@@ -55,7 +55,7 @@ void SGraphPinAI::Construct(const FArguments& InArgs, UEdGraphPin* InPin)
 		.BorderBackgroundColor(this, &SGraphPinAI::GetPinColor)
 		.OnMouseButtonDown(this, &SGraphPinAI::OnPinMouseDown)
 		.Cursor(this, &SGraphPinAI::GetPinCursor)
-		.Padding(FMargin(5.0f))
+		.Padding(FMargin(10.0f))
 		);
 }
 
@@ -95,6 +95,13 @@ FText SGraphNodeAI::GetDescription() const
 	return MyNode ? MyNode->GetDescription() : FText::GetEmpty();
 }
 
+EVisibility SGraphNodeAI::GetDescriptionVisibility() const
+{
+	// LOD this out once things get too small
+	TSharedPtr<SGraphPanel> MyOwnerPanel = GetOwnerPanel();
+	return (!MyOwnerPanel.IsValid() || MyOwnerPanel->GetCurrentLOD() > EGraphRenderingLOD::LowDetail) ? EVisibility::Visible : EVisibility::Collapsed;
+}
+
 void SGraphNodeAI::AddPin(const TSharedRef<SGraphPin>& PinToAdd)
 {
 	PinToAdd->SetOwner(SharedThis(this));
diff --git a/Engine/Source/Editor/AIGraph/Public/SGraphNodeAI.h b/Engine/Source/Editor/AIGraph/Public/SGraphNodeAI.h
--- a/Engine/Source/Editor/AIGraph/Public/SGraphNodeAI.h
+++ b/Engine/Source/Editor/AIGraph/Public/SGraphNodeAI.h
@@ -62,6 +62,8 @@ protected:
 	uint32 bDragMarkerVisible : 1;
 
 	virtual FText GetDescription() const;
+	virtual EVisibility GetDescriptionVisibility() const;
+
 	virtual FText GetPreviewCornerText() const;
 	virtual const FSlateBrush* GetNameIcon() const;
 };
diff --git a/Engine/Source/Editor/BehaviorTreeEditor/BehaviorTreeEditor.Build.cs b/Engine/Source/Editor/BehaviorTreeEditor/BehaviorTreeEditor.Build.cs
--- a/Engine/Source/Editor/BehaviorTreeEditor/BehaviorTreeEditor.Build.cs
+++ b/Engine/Source/Editor/BehaviorTreeEditor/BehaviorTreeEditor.Build.cs
@@ -37,6 +37,7 @@ public class BehaviorTreeEditor : ModuleRules
 				"MessageLog", 
 				"GraphEditor",
                 "Kismet",
+				"KismetWidgets",
                 "PropertyEditor",
 				"AnimGraph",
 				"BlueprintGraph",
diff --git a/Engine/Source/Editor/BehaviorTreeEditor/Private/SGraphNode_BehaviorTree.cpp b/Engine/Source/Editor/BehaviorTreeEditor/Private/SGraphNode_BehaviorTree.cpp
--- a/Engine/Source/Editor/BehaviorTreeEditor/Private/SGraphNode_BehaviorTree.cpp
+++ b/Engine/Source/Editor/BehaviorTreeEditor/Private/SGraphNode_BehaviorTree.cpp
@@ -16,6 +16,7 @@
 #include "IDocumentation.h"
 #include "SInlineEditableTextBlock.h"
 #include "SCommentBubble.h"
+#include "SLevelOfDetailBranchNode.h"
 
 #define LOCTEXT_NAMESPACE "BehaviorTreeEditor"
 
@@ -318,12 +319,24 @@ void SGraphNode_BehaviorTree::UpdateGraphNode()
 	TSharedPtr<SErrorText> ErrorText;
 	TSharedPtr<STextBlock> DescriptionText; 
 	TSharedPtr<SNodeTitle> NodeTitle = SNew(SNodeTitle, GraphNode);
-	float NodePadding = 10.0f;
 
-	if (Cast<UBehaviorTreeGraphNode_Decorator>(GraphNode) || Cast<UBehaviorTreeGraphNode_CompositeDecorator>(GraphNode) || Cast<UBehaviorTreeGraphNode_Service>(GraphNode))
+	TWeakPtr<SNodeTitle> WeakNodeTitle = NodeTitle;
+	auto GetNodeTitlePlaceholderWidth = [WeakNodeTitle]() -> FOptionalSize
 	{
-		NodePadding = 2.0f;
-	}
+		TSharedPtr<SNodeTitle> NodeTitlePin = WeakNodeTitle.Pin();
+		const float DesiredWidth = (NodeTitlePin.IsValid()) ? NodeTitlePin->GetTitleSize().X : 0.0f;
+		return FMath::Max(75.0f, DesiredWidth);
+	};
+	auto GetNodeTitlePlaceholderHeight = [WeakNodeTitle]() -> FOptionalSize
+	{
+		TSharedPtr<SNodeTitle> NodeTitlePin = WeakNodeTitle.Pin();
+		const float DesiredHeight = (NodeTitlePin.IsValid()) ? NodeTitlePin->GetTitleSize().Y : 0.0f;
+		return FMath::Max(22.0f, DesiredHeight);
+	};
+
+	const FMargin NodePadding = (Cast<UBehaviorTreeGraphNode_Decorator>(GraphNode) || Cast<UBehaviorTreeGraphNode_CompositeDecorator>(GraphNode) || Cast<UBehaviorTreeGraphNode_Service>(GraphNode))
+		? FMargin(2.0f)
+		: FMargin(8.0f);
 
 	IndexOverlay = SNew(SBehaviorTreeIndex)
 		.ToolTipText(this, &SGraphNode_BehaviorTree::GetIndexTooltipText)
@@ -339,132 +352,164 @@ void SGraphNode_BehaviorTree::UpdateGraphNode()
 		[
 			SNew(SBorder)
 			.BorderImage( FEditorStyle::GetBrush( "Graph.StateNode.Body" ) )
-			.Padding(0)
+			.Padding(0.0f)
 			.BorderBackgroundColor( this, &SGraphNode_BehaviorTree::GetBorderBackgroundColor )
 			.OnMouseButtonDown(this, &SGraphNode_BehaviorTree::OnMouseDown)
 			[
-				SNew(SOverlay)	
-				// INPUT PIN AREA
-				+SOverlay::Slot()
-				.HAlign(HAlign_Fill)
-				.VAlign(VAlign_Top)
-				[
-					SAssignNew(LeftNodeBox, SVerticalBox)
-				]
-
-				// OUTPUT PIN AREA
-				+SOverlay::Slot()
-				.HAlign(HAlign_Fill)
-				.VAlign(VAlign_Bottom)
-				[
-					SAssignNew(RightNodeBox, SVerticalBox)
-					+SVerticalBox::Slot()
-					.HAlign(HAlign_Fill)
-					.VAlign(VAlign_Fill)
-					.Padding(20.0f,0.0f)
-					.FillHeight(1.0f)
-					[
-						SAssignNew(OutputPinBox, SHorizontalBox)
-					]
-				]
+				SNew(SOverlay)
 
-				// STATE NAME AREA
+				// Pins and node details
 				+SOverlay::Slot()
 				.HAlign(HAlign_Fill)
-				.VAlign(VAlign_Center)
-				.Padding(NodePadding)
+				.VAlign(VAlign_Fill)
 				[
 					SNew(SVerticalBox)
+
+					// INPUT PIN AREA
 					+SVerticalBox::Slot()
 					.AutoHeight()
 					[
-						DecoratorsBox.ToSharedRef()
+						SNew(SBox)
+						.MinDesiredHeight(NodePadding.Top)
+						[
+							SAssignNew(LeftNodeBox, SVerticalBox)
+						]
 					]
+
+					// STATE NAME AREA
 					+SVerticalBox::Slot()
-					.AutoHeight()
+					.Padding(FMargin(NodePadding.Left, 0.0f, NodePadding.Right, 0.0f))
 					[
-						SAssignNew(NodeBody, SBorder)
-						.BorderImage( FEditorStyle::GetBrush("BTEditor.Graph.BTNode.Body") )
-						.BorderBackgroundColor( this, &SGraphNode_BehaviorTree::GetBackgroundColor )
-						.HAlign(HAlign_Fill)
-						.VAlign(VAlign_Center)
-						.Visibility(EVisibility::SelfHitTestInvisible)
+						SNew(SVerticalBox)
+						+SVerticalBox::Slot()
+						.AutoHeight()
+						[
+							DecoratorsBox.ToSharedRef()
+						]
+						+SVerticalBox::Slot()
+						.AutoHeight()
 						[
-							SNew(SOverlay)
-							+SOverlay::Slot()
+							SAssignNew(NodeBody, SBorder)
+							.BorderImage( FEditorStyle::GetBrush("BTEditor.Graph.BTNode.Body") )
+							.BorderBackgroundColor( this, &SGraphNode_BehaviorTree::GetBackgroundColor )
 							.HAlign(HAlign_Fill)
-							.VAlign(VAlign_Fill)
+							.VAlign(VAlign_Center)
+							.Visibility(EVisibility::SelfHitTestInvisible)
 							[
-								SNew(SVerticalBox)
-								+SVerticalBox::Slot()
-								.AutoHeight()
+								SNew(SOverlay)
+								+SOverlay::Slot()
+								.HAlign(HAlign_Fill)
+								.VAlign(VAlign_Fill)
 								[
-									SNew(SHorizontalBox)
-									+SHorizontalBox::Slot()
-									.AutoWidth()
-									[
-										// POPUP ERROR MESSAGE
-										SAssignNew(ErrorText, SErrorText )
-										.BackgroundColor( this, &SGraphNode_BehaviorTree::GetErrorColor )
-										.ToolTipText( this, &SGraphNode_BehaviorTree::GetErrorMsgToolTip )
-									]
-									+SHorizontalBox::Slot()
-									.AutoWidth()
-									.VAlign(VAlign_Center)
+									SNew(SVerticalBox)
+									+SVerticalBox::Slot()
+									.AutoHeight()
 									[
-										SNew(SImage)
-										.Image(this, &SGraphNode_BehaviorTree::GetNameIcon)
-									]
-									+SHorizontalBox::Slot()
-									.Padding(FMargin(4.0f, 0.0f, 4.0f, 0.0f))
-									[
-										SNew(SVerticalBox)
-										+SVerticalBox::Slot()
-										.AutoHeight()
+										SNew(SHorizontalBox)
+										+SHorizontalBox::Slot()
+										.AutoWidth()
 										[
-											SAssignNew(InlineEditableText, SInlineEditableTextBlock)
-											.Style( FEditorStyle::Get(), "Graph.StateNode.NodeTitleInlineEditableText" )
-											.Text( NodeTitle.Get(), &SNodeTitle::GetHeadTitle )
-											.OnVerifyTextChanged(this, &SGraphNode_BehaviorTree::OnVerifyNameTextChanged)
-											.OnTextCommitted(this, &SGraphNode_BehaviorTree::OnNameTextCommited)
-											.IsReadOnly( this, &SGraphNode_BehaviorTree::IsNameReadOnly )
-											.IsSelected(this, &SGraphNode_BehaviorTree::IsSelectedExclusively)
+											// POPUP ERROR MESSAGE
+											SAssignNew(ErrorText, SErrorText )
+											.BackgroundColor( this, &SGraphNode_BehaviorTree::GetErrorColor )
+											.ToolTipText( this, &SGraphNode_BehaviorTree::GetErrorMsgToolTip )
 										]
-										+SVerticalBox::Slot()
-										.AutoHeight()
+										
+										+SHorizontalBox::Slot()
+										.AutoWidth()
 										[
-											NodeTitle.ToSharedRef()
+											SNew(SLevelOfDetailBranchNode)
+											.UseLowDetailSlot(this, &SGraphNode_BehaviorTree::UseLowDetailNodeTitles)
+											.LowDetail()
+											[
+												SNew(SBox)
+												.WidthOverride_Lambda(GetNodeTitlePlaceholderWidth)
+												.HeightOverride_Lambda(GetNodeTitlePlaceholderHeight)
+											]
+											.HighDetail()
+											[
+												SNew(SHorizontalBox)
+												+SHorizontalBox::Slot()
+												.AutoWidth()
+												.VAlign(VAlign_Center)
+												[
+													SNew(SImage)
+													.Image(this, &SGraphNode_BehaviorTree::GetNameIcon)
+												]
+												+SHorizontalBox::Slot()
+												.Padding(FMargin(4.0f, 0.0f, 4.0f, 0.0f))
+												[
+													SNew(SVerticalBox)
+													+SVerticalBox::Slot()
+													.AutoHeight()
+													[
+														SAssignNew(InlineEditableText, SInlineEditableTextBlock)
+														.Style( FEditorStyle::Get(), "Graph.StateNode.NodeTitleInlineEditableText" )
+														.Text( NodeTitle.Get(), &SNodeTitle::GetHeadTitle )
+														.OnVerifyTextChanged(this, &SGraphNode_BehaviorTree::OnVerifyNameTextChanged)
+														.OnTextCommitted(this, &SGraphNode_BehaviorTree::OnNameTextCommited)
+														.IsReadOnly( this, &SGraphNode_BehaviorTree::IsNameReadOnly )
+														.IsSelected(this, &SGraphNode_BehaviorTree::IsSelectedExclusively)
+													]
+													+SVerticalBox::Slot()
+													.AutoHeight()
+													[
+														NodeTitle.ToSharedRef()
+													]
+												]
+											]
 										]
 									]
+									+SVerticalBox::Slot()
+									.AutoHeight()
+									[
+										// DESCRIPTION MESSAGE
+										SAssignNew(DescriptionText, STextBlock )
+										.Visibility(this, &SGraphNode_BehaviorTree::GetDescriptionVisibility)
+										.Text(this, &SGraphNode_BehaviorTree::GetDescription)
+									]
 								]
-								+SVerticalBox::Slot()
-								.AutoHeight()
+								+SOverlay::Slot()
+								.HAlign(HAlign_Right)
+								.VAlign(VAlign_Fill)
 								[
-									// DESCRIPTION MESSAGE
-									SAssignNew(DescriptionText, STextBlock )
-									.Text(this, &SGraphNode_BehaviorTree::GetDescription)
+									SNew(SBorder)
+									.BorderImage( FEditorStyle::GetBrush("BTEditor.Graph.BTNode.Body") )
+									.BorderBackgroundColor(BehaviorTreeColors::Debugger::SearchFailed)
+									.Padding(FMargin(4.0f, 0.0f))
+									.Visibility(this, &SGraphNode_BehaviorTree::GetDebuggerSearchFailedMarkerVisibility)
 								]
 							]
-							+SOverlay::Slot()
-							.HAlign(HAlign_Right)
-							.VAlign(VAlign_Fill)
-							[
-								SNew(SBorder)
-								.BorderImage( FEditorStyle::GetBrush("BTEditor.Graph.BTNode.Body") )
-								.BorderBackgroundColor(BehaviorTreeColors::Debugger::SearchFailed)
-								.Padding(FMargin(4.0f, 0.0f))
-								.Visibility(this, &SGraphNode_BehaviorTree::GetDebuggerSearchFailedMarkerVisibility)
-							]
+						]
+						+SVerticalBox::Slot()
+						.AutoHeight()
+						.Padding(FMargin(10.0f,0,0,0))
+						[
+							ServicesBox.ToSharedRef()
 						]
 					]
+
+					// OUTPUT PIN AREA
 					+SVerticalBox::Slot()
 					.AutoHeight()
-					.Padding(FMargin(10.0f,0,0,0))
 					[
-						ServicesBox.ToSharedRef()
+						SNew(SBox)
+						.MinDesiredHeight(NodePadding.Bottom)
+						[
+							SAssignNew(RightNodeBox, SVerticalBox)
+							+SVerticalBox::Slot()
+							.HAlign(HAlign_Fill)
+							.VAlign(VAlign_Fill)
+							.Padding(20.0f,0.0f)
+							.FillHeight(1.0f)
+							[
+								SAssignNew(OutputPinBox, SHorizontalBox)
+							]
+						]
 					]
 				]
-				//drag marker overlay
+
+				// Drag marker overlay
 				+SOverlay::Slot()
 				.HAlign(HAlign_Fill)
 				.VAlign(VAlign_Top)
@@ -479,6 +524,7 @@ void SGraphNode_BehaviorTree::UpdateGraphNode()
 						.HeightOverride(4)
 					]
 				]
+
 				// Blueprint indicator overlay
 				+SOverlay::Slot()
 				.HAlign(HAlign_Left)
@@ -687,18 +733,18 @@ void SGraphNode_BehaviorTree::AddPin(const TSharedRef<SGraphPin>& PinToAdd)
 			OutputPinBox->AddSlot()
 			.HAlign(HAlign_Fill)
 			.VAlign(VAlign_Fill)
-				.FillWidth(0.4f)
-				.Padding(0,0,20.0f,0)
-				[
-					PinToAdd
-				];
+			.FillWidth(0.4f)
+			.Padding(0,0,20.0f,0)
+			[
+				PinToAdd
+			];
 		}
 		else
 		{
 			OutputPinBox->AddSlot()
-				.HAlign(HAlign_Fill)
-				.VAlign(VAlign_Fill)
-				.FillWidth(1.0f)
+			.HAlign(HAlign_Fill)
+			.VAlign(VAlign_Fill)
+			.FillWidth(1.0f)
 			[
 				PinToAdd
 			];
@@ -832,8 +878,11 @@ EVisibility SGraphNode_BehaviorTree::GetIndexVisibility() const
 	}
 
 	// Visible if we are in PIE or if we have siblings
-	const bool bIsVisible = (GEditor->bIsSimulatingInEditor || GEditor->PlayWorld != NULL) || (MyParentOutputPin && MyParentOutputPin->LinkedTo.Num() > 1);
-	return bIsVisible ? EVisibility::Visible : EVisibility::Collapsed;
+	const bool bCanShowIndex = (GEditor->bIsSimulatingInEditor || GEditor->PlayWorld != NULL) || (MyParentOutputPin && MyParentOutputPin->LinkedTo.Num() > 1);
+
+	// LOD this out once things get too small
+	TSharedPtr<SGraphPanel> MyOwnerPanel = GetOwnerPanel();
+	return (bCanShowIndex && (!MyOwnerPanel.IsValid() || MyOwnerPanel->GetCurrentLOD() > EGraphRenderingLOD::LowDetail)) ? EVisibility::Visible : EVisibility::Collapsed;
 }
 
 FText SGraphNode_BehaviorTree::GetIndexText() const
@@ -894,7 +943,11 @@ FText SGraphNode_BehaviorTree::GetIndexTooltipText() const
 EVisibility SGraphNode_BehaviorTree::GetBlueprintIconVisibility() const
 {
 	UBehaviorTreeGraphNode* BTGraphNode = Cast<UBehaviorTreeGraphNode>(GraphNode);
-	return (BTGraphNode != nullptr && BTGraphNode->UsesBlueprint()) ? EVisibility::Visible : EVisibility::Collapsed;
+	const bool bCanShowIcon = (BTGraphNode != nullptr && BTGraphNode->UsesBlueprint());
+
+	// LOD this out once things get too small
+	TSharedPtr<SGraphPanel> MyOwnerPanel = GetOwnerPanel();
+	return (bCanShowIcon && (!MyOwnerPanel.IsValid() || MyOwnerPanel->GetCurrentLOD() > EGraphRenderingLOD::LowDetail)) ? EVisibility::Visible : EVisibility::Collapsed;
 }
 
 void SGraphNode_BehaviorTree::GetOverlayBrushes(bool bSelected, const FVector2D WidgetSize, TArray<FOverlayBrushInfo>& Brushes) const
diff --git a/Engine/Source/Editor/EnvironmentQueryEditor/EnvironmentQueryEditor.Build.cs b/Engine/Source/Editor/EnvironmentQueryEditor/EnvironmentQueryEditor.Build.cs
--- a/Engine/Source/Editor/EnvironmentQueryEditor/EnvironmentQueryEditor.Build.cs
+++ b/Engine/Source/Editor/EnvironmentQueryEditor/EnvironmentQueryEditor.Build.cs
@@ -35,6 +35,7 @@ public class EnvironmentQueryEditor : ModuleRules
 				"UnrealEd", 
 				"MessageLog", 
 				"GraphEditor",
+				"KismetWidgets",
                 "PropertyEditor",
 				"AnimGraph",
 				"BlueprintGraph",
diff --git a/Engine/Source/Editor/EnvironmentQueryEditor/Private/SGraphNode_EnvironmentQuery.cpp b/Engine/Source/Editor/EnvironmentQueryEditor/Private/SGraphNode_EnvironmentQuery.cpp
--- a/Engine/Source/Editor/EnvironmentQueryEditor/Private/SGraphNode_EnvironmentQuery.cpp
+++ b/Engine/Source/Editor/EnvironmentQueryEditor/Private/SGraphNode_EnvironmentQuery.cpp
@@ -13,6 +13,7 @@
 #include "EnvironmentQueryColors.h"
 #include "EnvironmentQuery/EnvQueryTest.h"
 #include "SInlineEditableTextBlock.h"
+#include "SLevelOfDetailBranchNode.h"
 
 #define LOCTEXT_NAMESPACE "EnvironmentQueryEditor"
 
@@ -111,7 +112,9 @@ void SGraphNode_EnvironmentQuery::UpdateGraphNode()
 	LeftNodeBox.Reset();
 	SubNodes.Reset();
 
-	const float NodePadding = (Cast<UEnvironmentQueryGraphNode_Test>(GraphNode) != NULL) ? 2.0f : 10.0f;
+	const FMargin NodePadding = (Cast<UEnvironmentQueryGraphNode_Test>(GraphNode) != NULL) 
+		? FMargin(2.0f) 
+		: FMargin(8.0f);
 	float TestPadding = 0.0f;
 
 	UEnvironmentQueryGraphNode_Option* OptionNode = Cast<UEnvironmentQueryGraphNode_Option>(GraphNode);
@@ -150,6 +153,20 @@ void SGraphNode_EnvironmentQuery::UpdateGraphNode()
 	TSharedPtr<STextBlock> DescriptionText; 
 	TSharedPtr<SNodeTitle> NodeTitle = SNew(SNodeTitle, GraphNode);
 
+	TWeakPtr<SNodeTitle> WeakNodeTitle = NodeTitle;
+	auto GetNodeTitlePlaceholderWidth = [WeakNodeTitle]() -> FOptionalSize
+	{
+		TSharedPtr<SNodeTitle> NodeTitlePin = WeakNodeTitle.Pin();
+		const float DesiredWidth = (NodeTitlePin.IsValid()) ? NodeTitlePin->GetTitleSize().X : 0.0f;
+		return FMath::Max(75.0f, DesiredWidth);
+	};
+	auto GetNodeTitlePlaceholderHeight = [WeakNodeTitle]() -> FOptionalSize
+	{
+		TSharedPtr<SNodeTitle> NodeTitlePin = WeakNodeTitle.Pin();
+		const float DesiredHeight = (NodeTitlePin.IsValid()) ? NodeTitlePin->GetTitleSize().Y : 0.0f;
+		return FMath::Max(22.0f, DesiredHeight);
+	};
+
 	this->ContentScale.Bind( this, &SGraphNode::GetContentScale );
 	this->GetOrAddSlot( ENodeZone::Center )
 		.HAlign(HAlign_Fill)
@@ -163,27 +180,27 @@ void SGraphNode_EnvironmentQuery::UpdateGraphNode()
 			[
 				SNew(SOverlay)
 
-				// INPUT PIN AREA
+				// Pins and node details
 				+SOverlay::Slot()
-					.HAlign(HAlign_Fill)
-					.VAlign(VAlign_Top)
-					[
-						SAssignNew(LeftNodeBox, SVerticalBox)
-					]
+				.HAlign(HAlign_Fill)
+				.VAlign(VAlign_Fill)
+				[
+					SNew(SVerticalBox)
 
-				// OUTPUT PIN AREA
-				+SOverlay::Slot()
-					.HAlign(HAlign_Fill)
-					.VAlign(VAlign_Bottom)
+					// INPUT PIN AREA
+					+SVerticalBox::Slot()
+					.AutoHeight()
 					[
-						SAssignNew(RightNodeBox, SVerticalBox)
+						SNew(SBox)
+						.MinDesiredHeight(NodePadding.Top)
+						[
+							SAssignNew(LeftNodeBox, SVerticalBox)
+						]
 					]
 
-				// STATE NAME AREA
-				+SOverlay::Slot()
-					.HAlign(HAlign_Fill)
-					.VAlign(VAlign_Center)
-					.Padding(NodePadding)
+					// STATE NAME AREA
+					+SVerticalBox::Slot()
+					.Padding(FMargin(NodePadding.Left, 0.0f, NodePadding.Right, 0.0f))
 					[
 						SNew(SHorizontalBox)
 						+SHorizontalBox::Slot()
@@ -211,60 +228,72 @@ void SGraphNode_EnvironmentQuery::UpdateGraphNode()
 								.Visibility(EVisibility::SelfHitTestInvisible)
 								[
 									SNew(SVerticalBox)
-										+SVerticalBox::Slot()
-										.AutoHeight()
-										.Padding(0,0,0,2)
+									+SVerticalBox::Slot()
+									.AutoHeight()
+									.Padding(0,0,0,2)
+									[
+										SNew(SBox).HeightOverride(4)
 										[
-											SNew(SBox).HeightOverride(4)
-											[
-												// weight bar
-												SNew(SProgressBar)
-												.FillColorAndOpacity(this, &SGraphNode_EnvironmentQuery::GetWeightProgressBarColor)
-												.Visibility(this, &SGraphNode_EnvironmentQuery::GetWeightMarkerVisibility)
-												.Percent(this, &SGraphNode_EnvironmentQuery::GetWeightProgressBarPercent)
-											]
+											// weight bar
+											SNew(SProgressBar)
+											.FillColorAndOpacity(this, &SGraphNode_EnvironmentQuery::GetWeightProgressBarColor)
+											.Visibility(this, &SGraphNode_EnvironmentQuery::GetWeightMarkerVisibility)
+											.Percent(this, &SGraphNode_EnvironmentQuery::GetWeightProgressBarPercent)
 										]
-										+SVerticalBox::Slot()
-										.AutoHeight()
+									]
+									+SVerticalBox::Slot()
+									.AutoHeight()
+									[
+										SNew(SHorizontalBox)
+										+SHorizontalBox::Slot()
+										.AutoWidth()
+										[
+											// POPUP ERROR MESSAGE
+											SAssignNew(ErrorText, SErrorText )
+											.BackgroundColor( this, &SGraphNode_EnvironmentQuery::GetErrorColor )
+											.ToolTipText( this, &SGraphNode_EnvironmentQuery::GetErrorMsgToolTip )
+										]
+										+SHorizontalBox::Slot()
+										.Padding(FMargin(4.0f, 0.0f, 4.0f, 0.0f))
+										[
+											SNew(SLevelOfDetailBranchNode)
+											.UseLowDetailSlot(this, &SGraphNode_EnvironmentQuery::UseLowDetailNodeTitles)
+											.LowDetail()
 											[
-												SNew(SHorizontalBox)
-												+SHorizontalBox::Slot()
-													.AutoWidth()
-													[
-														// POPUP ERROR MESSAGE
-														SAssignNew(ErrorText, SErrorText )
-														.BackgroundColor( this, &SGraphNode_EnvironmentQuery::GetErrorColor )
-														.ToolTipText( this, &SGraphNode_EnvironmentQuery::GetErrorMsgToolTip )
-													]
-												+SHorizontalBox::Slot()
-													.Padding(FMargin(4.0f, 0.0f, 4.0f, 0.0f))
-													[
-														SNew(SVerticalBox)
-														+SVerticalBox::Slot()
-														.AutoHeight()
-															[
-																SAssignNew(InlineEditableText, SInlineEditableTextBlock)
-																.Style( FEditorStyle::Get(), "Graph.StateNode.NodeTitleInlineEditableText" )
-																.Text( NodeTitle.Get(), &SNodeTitle::GetHeadTitle )
-																.OnVerifyTextChanged(this, &SGraphNode_EnvironmentQuery::OnVerifyNameTextChanged)
-																.OnTextCommitted(this, &SGraphNode_EnvironmentQuery::OnNameTextCommited)
-																.IsReadOnly( this, &SGraphNode_EnvironmentQuery::IsNameReadOnly )
-																.IsSelected(this, &SGraphNode_EnvironmentQuery::IsSelectedExclusively)
-															]
-														+SVerticalBox::Slot()
-															.AutoHeight()
-															[
-																NodeTitle.ToSharedRef()
-															]
-													]
+												SNew(SBox)
+												.WidthOverride_Lambda(GetNodeTitlePlaceholderWidth)
+												.HeightOverride_Lambda(GetNodeTitlePlaceholderHeight)
 											]
-										+SVerticalBox::Slot()
-											.AutoHeight()
+											.HighDetail()
 											[
-												// DESCRIPTION MESSAGE
-												SAssignNew(DescriptionText, STextBlock )
-												.Text(this, &SGraphNode_EnvironmentQuery::GetDescription)
+												SNew(SVerticalBox)
+												+SVerticalBox::Slot()
+												.AutoHeight()
+												[
+													SAssignNew(InlineEditableText, SInlineEditableTextBlock)
+													.Style( FEditorStyle::Get(), "Graph.StateNode.NodeTitleInlineEditableText" )
+													.Text( NodeTitle.Get(), &SNodeTitle::GetHeadTitle )
+													.OnVerifyTextChanged(this, &SGraphNode_EnvironmentQuery::OnVerifyNameTextChanged)
+													.OnTextCommitted(this, &SGraphNode_EnvironmentQuery::OnNameTextCommited)
+													.IsReadOnly( this, &SGraphNode_EnvironmentQuery::IsNameReadOnly )
+													.IsSelected(this, &SGraphNode_EnvironmentQuery::IsSelectedExclusively)
+												]
+												+SVerticalBox::Slot()
+												.AutoHeight()
+												[
+													NodeTitle.ToSharedRef()
+												]
 											]
+										]
+									]
+									+SVerticalBox::Slot()
+									.AutoHeight()
+									[
+										// DESCRIPTION MESSAGE
+										SAssignNew(DescriptionText, STextBlock )
+										.Visibility(this, &SGraphNode_EnvironmentQuery::GetDescriptionVisibility)
+										.Text(this, &SGraphNode_EnvironmentQuery::GetDescription)
+									]
 								]
 							]
 							+SVerticalBox::Slot()
@@ -275,21 +304,34 @@ void SGraphNode_EnvironmentQuery::UpdateGraphNode()
 							]
 						]
 					]
-				//drag marker overlay
-				+SOverlay::Slot()
-					.HAlign(HAlign_Fill)
-					.VAlign(VAlign_Top)
+
+					// OUTPUT PIN AREA
+					+SVerticalBox::Slot()
+					.AutoHeight()
 					[
-						SNew(SBorder)
-						.BorderBackgroundColor(EnvironmentQueryColors::Action::DragMarker)
-						.ColorAndOpacity(EnvironmentQueryColors::Action::DragMarker)
-						.BorderImage(FEditorStyle::GetBrush("Graph.StateNode.Body"))
-						.Visibility(this, &SGraphNode_EnvironmentQuery::GetDragOverMarkerVisibility)
+						SNew(SBox)
+						.MinDesiredHeight(NodePadding.Bottom)
 						[
-							SNew(SBox)
-							.HeightOverride(4)
+							SAssignNew(RightNodeBox, SVerticalBox)
 						]
 					]
+				]
+
+				// Drag marker overlay
+				+SOverlay::Slot()
+				.HAlign(HAlign_Fill)
+				.VAlign(VAlign_Top)
+				[
+					SNew(SBorder)
+					.BorderBackgroundColor(EnvironmentQueryColors::Action::DragMarker)
+					.ColorAndOpacity(EnvironmentQueryColors::Action::DragMarker)
+					.BorderImage(FEditorStyle::GetBrush("Graph.StateNode.Body"))
+					.Visibility(this, &SGraphNode_EnvironmentQuery::GetDragOverMarkerVisibility)
+					[
+						SNew(SBox)
+						.HeightOverride(4)
+					]
+				]
 			]
 		];
 
diff --git a/Engine/Source/Editor/GraphEditor/Private/ConnectionDrawingPolicy.cpp b/Engine/Source/Editor/GraphEditor/Private/ConnectionDrawingPolicy.cpp
--- a/Engine/Source/Editor/GraphEditor/Private/ConnectionDrawingPolicy.cpp
+++ b/Engine/Source/Editor/GraphEditor/Private/ConnectionDrawingPolicy.cpp
@@ -479,8 +479,8 @@ bool FConnectionDrawingPolicy::IsConnectionCulled( const FArrangedWidget& StartL
 {
 	const float Top		= FMath::Min( StartLink.Geometry.AbsolutePosition.Y, EndLink.Geometry.AbsolutePosition.Y );
 	const float Left	= FMath::Min( StartLink.Geometry.AbsolutePosition.X, EndLink.Geometry.AbsolutePosition.X );
-	const float Bottom	= FMath::Max( StartLink.Geometry.AbsolutePosition.Y, EndLink.Geometry.AbsolutePosition.Y );
-	const float Right	= FMath::Max( StartLink.Geometry.AbsolutePosition.X, EndLink.Geometry.AbsolutePosition.X ); 
+	const float Bottom	= FMath::Max( StartLink.Geometry.AbsolutePosition.Y + StartLink.Geometry.Size.Y, EndLink.Geometry.AbsolutePosition.Y + EndLink.Geometry.Size.Y );
+	const float Right	= FMath::Max( StartLink.Geometry.AbsolutePosition.X + StartLink.Geometry.Size.X, EndLink.Geometry.AbsolutePosition.X + EndLink.Geometry.Size.X ); 
 
 	return	Left > ClippingRect.Right || Right < ClippingRect.Left || 
 			Bottom < ClippingRect.Top || Top > ClippingRect.Bottom;
diff --git a/Engine/Source/Editor/GraphEditor/Private/SGraphNode.cpp b/Engine/Source/Editor/GraphEditor/Private/SGraphNode.cpp
--- a/Engine/Source/Editor/GraphEditor/Private/SGraphNode.cpp
+++ b/Engine/Source/Editor/GraphEditor/Private/SGraphNode.cpp
@@ -22,6 +22,8 @@ void SNodeTitle::Construct(const FArguments& InArgs, UEdGraphNode* InNode)
 
 	ExtraLineStyle = InArgs._ExtraLineStyle;
 
+	CachedSize = FVector2D::ZeroVector;
+
 	// If the user set the text, use it, otherwise use the node title by default
 	if(InArgs._Text.IsSet())
 	{
@@ -37,6 +39,8 @@ void SNodeTitle::Construct(const FArguments& InArgs, UEdGraphNode* InNode)
 
 void SNodeTitle::Tick( const FGeometry& AllottedGeometry, const double InCurrentTime, const float InDeltaTime )
 {
+	CachedSize = AllottedGeometry.GetLocalSize();
+
 	// Checks to see if the cached string is valid, and if not, updates it.
 	if (NodeTitleCache.IsOutOfDate(GraphNode))
 	{
@@ -57,6 +61,11 @@ FText SNodeTitle::GetHeadTitle() const
 	return GraphNode->bCanRenameNode ? GraphNode->GetNodeTitle(ENodeTitleType::EditableTitle) : CachedHeadTitle;
 }
 
+FVector2D SNodeTitle::GetTitleSize() const
+{
+	return CachedSize;
+}
+
 void SNodeTitle::RebuildWidget()
 {
 	// Create the box to contain the lines
diff --git a/Engine/Source/Editor/GraphEditor/Public/SGraphNode.h b/Engine/Source/Editor/GraphEditor/Public/SGraphNode.h
--- a/Engine/Source/Editor/GraphEditor/Public/SGraphNode.h
+++ b/Engine/Source/Editor/GraphEditor/Public/SGraphNode.h
@@ -39,6 +39,9 @@ public:
 	/** Returns the main title for the node */
 	FText GetHeadTitle() const;
 
+	/** Get the size of this title the last time it was drawn */
+	FVector2D GetTitleSize() const;
+
 protected:
 	UEdGraphNode* GraphNode;
 	FNodeTextCache NodeTitleCache;
@@ -50,6 +53,9 @@ protected:
 	/** The title text to use, auto-binds to get the title if not set externally */
 	TAttribute< FText > TitleText;
 
+	/** The cached size of the title */
+	FVector2D CachedSize;
+
 protected:
 	// Gets the expected node title
 	FText GetNodeTitle() const;
