Behavior Tree debugger improvements

UETOOL-298 - The "step into"/out/over buttons in the Behavior Tree enable and disable all the time

The "Back: Over" and "Forward: Over" step options are now able to perform a step that means they move back up to a parent tree. Previously this was being disallowed as the they were incorrectly treating that as entering a sub-tree.

If you pause your PIE session, the behavior tree debugger will now make sure and open the correct graph. Previously you could pause while within a sub-tree and the view would be left on the parent graph.

diff --git a/Engine/Source/Editor/BehaviorTreeEditor/Private/BehaviorTreeDebugger.cpp b/Engine/Source/Editor/BehaviorTreeEditor/Private/BehaviorTreeDebugger.cpp
--- a/Engine/Source/Editor/BehaviorTreeEditor/Private/BehaviorTreeDebugger.cpp
+++ b/Engine/Source/Editor/BehaviorTreeEditor/Private/BehaviorTreeDebugger.cpp
@@ -26,6 +26,7 @@ FBehaviorTreeDebugger::FBehaviorTreeDebugger()
 
 	FEditorDelegates::BeginPIE.AddRaw(this, &FBehaviorTreeDebugger::OnBeginPIE);
 	FEditorDelegates::EndPIE.AddRaw(this, &FBehaviorTreeDebugger::OnEndPIE);
+	FEditorDelegates::PausePIE.AddRaw(this, &FBehaviorTreeDebugger::OnPausePIE);
 
 #if USE_BEHAVIORTREE_DEBUGGER
 	UBehaviorTreeComponent::ActiveDebuggerCounter++;
@@ -36,6 +37,7 @@ FBehaviorTreeDebugger::~FBehaviorTreeDebugger()
 {
 	FEditorDelegates::BeginPIE.RemoveAll(this);
 	FEditorDelegates::EndPIE.RemoveAll(this);
+	FEditorDelegates::PausePIE.RemoveAll(this);
 	USelection::SelectObjectEvent.RemoveAll(this);
 	FBehaviorTreeDelegates::OnTreeStarted.RemoveAll(this);
 	FBehaviorTreeDelegates::OnDebugLocked.RemoveAll(this);
@@ -234,6 +236,23 @@ void FBehaviorTreeDebugger::OnEndPIE(const bool bIsSimulating)
 	UpdateDebuggerViewOnInstanceChange();
 }
 
+void FBehaviorTreeDebugger::OnPausePIE(const bool bIsSimulating)
+{
+#if USE_BEHAVIORTREE_DEBUGGER
+	// We might have paused while executing a sub-tree, so make sure that the editor is showing the correct tree
+	TSharedPtr<FBehaviorTreeEditor> EditorOwnerPin = EditorOwner.Pin();
+	if (EditorOwnerPin.IsValid() && TreeInstance.IsValid() && TreeInstance->DebuggerSteps.IsValidIndex(ActiveStepIndex))
+	{
+		const FBehaviorTreeExecutionStep& StepInfo = TreeInstance->DebuggerSteps[ActiveStepIndex];
+		const int32 LastInstanceIndex = StepInfo.InstanceStack.Num() - 1;
+		if (StepInfo.InstanceStack.IsValidIndex(LastInstanceIndex) && StepInfo.InstanceStack[LastInstanceIndex].TreeAsset != TreeAsset)
+		{
+			EditorOwnerPin->DebuggerSwitchAsset(StepInfo.InstanceStack[LastInstanceIndex].TreeAsset);
+		}
+	}
+#endif
+}
+
 void FBehaviorTreeDebugger::OnObjectSelected(UObject* Object)
 {
 	if (Object && Object->IsSelected())
@@ -947,6 +966,7 @@ void FBehaviorTreeDebugger::PausePlaySession()
 	if (GUnrealEd->PlayWorld)
 	{
 		GUnrealEd->PlayWorld->bDebugPauseExecution = true;
+		GUnrealEd->PlaySessionPaused();
 	}
 }
 
@@ -955,6 +975,7 @@ void FBehaviorTreeDebugger::ResumePlaySession()
 	if (GUnrealEd->PlayWorld)
 	{
 		GUnrealEd->PlayWorld->bDebugPauseExecution = false;
+		GUnrealEd->PlaySessionResumed();
 	}
 }
 
@@ -1215,7 +1236,7 @@ static int32 GetNumActiveInstances(const FBehaviorTreeExecutionStep& StepInfo, c
 {
 	for (int32 Idx = StepInfo.InstanceStack.Num() - 1; Idx >= 0; Idx--)
 	{
-		if (StepInfo.InstanceStack[Idx].ActivePath.Num())
+		//if (StepInfo.InstanceStack[Idx].ActivePath.Num())
 		{
 			ActiveSubtree = StepInfo.InstanceStack[Idx].TreeAsset;
 			return Idx + 1;
@@ -1235,16 +1256,17 @@ void FBehaviorTreeDebugger::UpdateAvailableActions()
 	StepOutIdx = INDEX_NONE;
 
 #if USE_BEHAVIORTREE_DEBUGGER
-	if (TreeInstance.IsValid() && TreeInstance->DebuggerSteps.IsValidIndex(ActiveStepIndex) && DebuggerInstanceIndex >= 0)
+	UBehaviorTreeComponent* TreeInstancePtr = TreeInstance.Get();
+	if (TreeInstancePtr && TreeInstancePtr->DebuggerSteps.IsValidIndex(ActiveStepIndex) && DebuggerInstanceIndex >= 0)
 	{
-		const FBehaviorTreeExecutionStep& CurStepInfo = TreeInstance->DebuggerSteps[ActiveStepIndex];
+		const FBehaviorTreeExecutionStep& CurStepInfo = TreeInstancePtr->DebuggerSteps[ActiveStepIndex];
 
-		if (TreeInstance->DebuggerSteps.IsValidIndex(ActiveStepIndex - 1))
+		if (TreeInstancePtr->DebuggerSteps.IsValidIndex(ActiveStepIndex - 1))
 		{
 			StepBackIntoIdx = ActiveStepIndex - 1;
 		}
 
-		if (TreeInstance->DebuggerSteps.IsValidIndex(ActiveStepIndex + 1))
+		if (TreeInstancePtr->DebuggerSteps.IsValidIndex(ActiveStepIndex + 1))
 		{
 			StepForwardIntoIdx = ActiveStepIndex + 1;
 		}
@@ -1255,55 +1277,41 @@ void FBehaviorTreeDebugger::UpdateAvailableActions()
 
 		for (int32 TestStepIndex = ActiveStepIndex - 1; TestStepIndex >= 0; TestStepIndex--)
 		{
-			const FBehaviorTreeExecutionStep& TestStepInfo = TreeInstance->DebuggerSteps[TestStepIndex];
+			const FBehaviorTreeExecutionStep& TestStepInfo = TreeInstancePtr->DebuggerSteps[TestStepIndex];
 			UBehaviorTree* TestTree = NULL;
 			const int32 TestStepInstances = GetNumActiveInstances(TestStepInfo, TestTree);
 
-			if (TestStepInstances < CurStepInstances ||
+			StepBackOverIdx = TestStepIndex;
+
+			// keep going only if the execution is moving to a sub-tree
+			if (TestStepInstances <= CurStepInstances ||
 				TestStepInfo.InstanceStack[DebuggerInstanceIndex].TreeAsset != CurStepInfo.InstanceStack[DebuggerInstanceIndex].TreeAsset)
 			{
-				// execution left current subtree
 				break;
 			}
-
-			if (TestStepInstances > CurStepInstances)
-			{
-				// next step is inside nested subtree
-				continue;
-			}
-
-			StepBackOverIdx = TestStepIndex;
-			break;
 		}
 
-		for (int32 TestStepIndex = ActiveStepIndex + 1; TestStepIndex < TreeInstance->DebuggerSteps.Num(); TestStepIndex++)
+		for (int32 TestStepIndex = ActiveStepIndex + 1; TestStepIndex < TreeInstancePtr->DebuggerSteps.Num(); TestStepIndex++)
 		{
-			const FBehaviorTreeExecutionStep& TestStepInfo = TreeInstance->DebuggerSteps[TestStepIndex];
+			const FBehaviorTreeExecutionStep& TestStepInfo = TreeInstancePtr->DebuggerSteps[TestStepIndex];
 			UBehaviorTree* TestTree = NULL;
 			int32 TestStepInstances = GetNumActiveInstances(TestStepInfo, TestTree);
 
-			if (TestStepInstances < CurStepInstances ||
+			StepForwardOverIdx = TestStepIndex;
+
+			// keep going only if the execution is moving to a sub-tree
+			if (TestStepInstances <= CurStepInstances ||
 				TestStepInfo.InstanceStack[DebuggerInstanceIndex].TreeAsset != CurStepInfo.InstanceStack[DebuggerInstanceIndex].TreeAsset)
 			{
-				// execution left current subtree
 				break;
 			}
-
-			if (TestStepInstances > CurStepInstances)
-			{
-				// next step is inside nested subtree
-				continue;
-			}
-
-			StepForwardOverIdx = TestStepIndex;
-			break;
 		}
 
 		if (CurStepInfo.InstanceStack.IsValidIndex(DebuggerInstanceIndex) && CurStepInfo.InstanceStack[DebuggerInstanceIndex].ActivePath.Num())
 		{
-			for (int32 TestStepIndex = ActiveStepIndex + 1; TestStepIndex < TreeInstance->DebuggerSteps.Num(); TestStepIndex++)
+			for (int32 TestStepIndex = ActiveStepIndex + 1; TestStepIndex < TreeInstancePtr->DebuggerSteps.Num(); TestStepIndex++)
 			{
-				const FBehaviorTreeExecutionStep& TestStepInfo = TreeInstance->DebuggerSteps[TestStepIndex];
+				const FBehaviorTreeExecutionStep& TestStepInfo = TreeInstancePtr->DebuggerSteps[TestStepIndex];
 				UBehaviorTree* TestTree = NULL;
 				int32 TestStepInstances = GetNumActiveInstances(TestStepInfo, TestTree);
 
diff --git a/Engine/Source/Editor/BehaviorTreeEditor/Private/BehaviorTreeDebugger.h b/Engine/Source/Editor/BehaviorTreeEditor/Private/BehaviorTreeDebugger.h
--- a/Engine/Source/Editor/BehaviorTreeEditor/Private/BehaviorTreeDebugger.h
+++ b/Engine/Source/Editor/BehaviorTreeEditor/Private/BehaviorTreeDebugger.h
@@ -35,6 +35,7 @@ public:
 	void OnTreeStarted(const UBehaviorTreeComponent& OwnerComp, const UBehaviorTree& InTreeAsset);
 	void OnBeginPIE(const bool bIsSimulating);
 	void OnEndPIE(const bool bIsSimulating);
+	void OnPausePIE(const bool bIsSimulating);
 
 	void OnBreakpointAdded(UBehaviorTreeGraphNode* Node);
 	void OnBreakpointRemoved(UBehaviorTreeGraphNode* Node);
