Fixed FBehaviorTreeDebugger::PausePlaySession and FBehaviorTreeDebugger::ResumePlaySession not guarding against re-entrancy

diff --git a/Engine/Source/Editor/BehaviorTreeEditor/Private/BehaviorTreeDebugger.cpp b/Engine/Source/Editor/BehaviorTreeEditor/Private/BehaviorTreeDebugger.cpp
--- a/Engine/Source/Editor/BehaviorTreeEditor/Private/BehaviorTreeDebugger.cpp
+++ b/Engine/Source/Editor/BehaviorTreeEditor/Private/BehaviorTreeDebugger.cpp
@@ -963,7 +963,7 @@ void FBehaviorTreeDebugger::StopPlaySession()
 
 void FBehaviorTreeDebugger::PausePlaySession()
 {
-	if (GUnrealEd->PlayWorld)
+	if (GUnrealEd->PlayWorld && !GUnrealEd->PlayWorld->bDebugPauseExecution)
 	{
 		GUnrealEd->PlayWorld->bDebugPauseExecution = true;
 		GUnrealEd->PlaySessionPaused();
@@ -972,7 +972,7 @@ void FBehaviorTreeDebugger::PausePlaySession()
 
 void FBehaviorTreeDebugger::ResumePlaySession()
 {
-	if (GUnrealEd->PlayWorld)
+	if (GUnrealEd->PlayWorld && GUnrealEd->PlayWorld->bDebugPauseExecution)
 	{
 		GUnrealEd->PlayWorld->bDebugPauseExecution = false;
 		GUnrealEd->PlaySessionResumed();
