Added a safety check to Blackboard Editor's key selection #UE4

UE-15693

diff --git a/Engine/Source/Editor/BehaviorTreeEditor/Private/BehaviorTreeEditor.cpp b/Engine/Source/Editor/BehaviorTreeEditor/Private/BehaviorTreeEditor.cpp
--- a/Engine/Source/Editor/BehaviorTreeEditor/Private/BehaviorTreeEditor.cpp
+++ b/Engine/Source/Editor/BehaviorTreeEditor/Private/BehaviorTreeEditor.cpp
@@ -319,7 +319,12 @@ void FBehaviorTreeEditor::HandleBlackboardEntrySelected(const FBlackboardEntry*
 {
 	// refresh details view
 	const bool bForceRefresh = true;
-	BlackboardDetailsView->SetObject(GetBlackboardData(), bForceRefresh);
+
+	if (ensure(BlackboardDetailsView.IsValid()))
+	{
+		// the opposite should never happen, we weren't able to internally repro it, but it seems someone was crashing on this line
+		BlackboardDetailsView->SetObject(GetBlackboardData(), bForceRefresh);
+	}
 }
 
 int32 FBehaviorTreeEditor::HandleGetSelectedBlackboardItemIndex(bool& bIsInherited)
