hidden deprecated blackboard key classes in editor's dropdown list
#ue4 UE-9682

diff --git a/Engine/Source/Editor/BehaviorTreeEditor/Private/SBehaviorTreeBlackboardEditor.cpp b/Engine/Source/Editor/BehaviorTreeEditor/Private/SBehaviorTreeBlackboardEditor.cpp
--- a/Engine/Source/Editor/BehaviorTreeEditor/Private/SBehaviorTreeBlackboardEditor.cpp
+++ b/Engine/Source/Editor/BehaviorTreeEditor/Private/SBehaviorTreeBlackboardEditor.cpp
@@ -174,7 +174,7 @@ public:
 	{
 		if(InClass != nullptr)
 		{
-			return !InClass->HasAnyClassFlags(CLASS_Abstract) &&
+			return !InClass->HasAnyClassFlags(CLASS_Abstract | CLASS_HideDropDown) &&
 				InClass->HasAnyClassFlags(CLASS_EditInlineNew) &&
 				InClass->IsChildOf(UBlackboardKeyType::StaticClass());
 		}
