Graph Editing: Renamed UEdGraphNode::NodeWidget DEPRECATED_NodeWidget to indicate it should not be used in future code since the assumption that at most a single graph editor is viewing a UEdGraph is not valid
(existing uses have been updated to the new name but not refactored yet)

diff --git a/Engine/Source/Editor/BehaviorTreeEditor/Private/BehaviorTreeGraph.cpp b/Engine/Source/Editor/BehaviorTreeEditor/Private/BehaviorTreeGraph.cpp
--- a/Engine/Source/Editor/BehaviorTreeEditor/Private/BehaviorTreeGraph.cpp
+++ b/Engine/Source/Editor/BehaviorTreeEditor/Private/BehaviorTreeGraph.cpp
@@ -936,7 +936,7 @@ void UBehaviorTreeGraph::SpawnMissingNodesForParallel()
 					UBehaviorTreeGraphNode* MainTaskNode = Cast<UBehaviorTreeGraphNode>(MainTaskPin->LinkedTo[0]->GetOwningNode());
 					if (MainTaskNode)
 					{
-						const int32 Width = MainTaskNode->NodeWidget.IsValid() ? MainTaskNode->NodeWidget.Pin()->GetDesiredSize().X : 200;
+						const int32 Width = MainTaskNode->DEPRECATED_NodeWidget.IsValid() ? MainTaskNode->DEPRECATED_NodeWidget.Pin()->GetDesiredSize().X : 200;
 						XOffset = MainTaskNode->NodePosX - ParallelNode->NodePosX + Width + 20;
 					}
 				}
@@ -946,7 +946,7 @@ void UBehaviorTreeGraph::SpawnMissingNodesForParallel()
 				WaitTaskNode->ClassData = FGraphNodeClassData(UBTTask_Wait::StaticClass(), "");
 				NodeBuilder.Finalize();
 
-				const int32 ParentHeight = ParallelNode->NodeWidget.IsValid() ? ParallelNode->NodeWidget.Pin()->GetDesiredSize().Y : 200;
+				const int32 ParentHeight = ParallelNode->DEPRECATED_NodeWidget.IsValid() ? ParallelNode->DEPRECATED_NodeWidget.Pin()->GetDesiredSize().Y : 200;
 				WaitTaskNode->NodePosX = ParallelNode->NodePosX + XOffset;
 				WaitTaskNode->NodePosY = ParallelNode->NodePosY + ParentHeight + 20;
 
@@ -1083,8 +1083,8 @@ namespace BTAutoArrangeHelpers
 				UBehaviorTreeGraphNode* GraphNode = Cast<UBehaviorTreeGraphNode>(Pin->LinkedTo[Idx]->GetOwningNode());
 				if (GraphNode && BBoxTree.Children.Num() > 0)
 				{
-					AutoArrangeNodes(GraphNode, BBoxTree.Children[BBoxIndex], PosX, PosY + GraphNode->NodeWidget.Pin()->GetDesiredSize().Y * 2.5f);
-					GraphNode->NodeWidget.Pin()->MoveTo(FVector2D(BBoxTree.Children[BBoxIndex].SubGraphBBox.X / 2 - GraphNode->NodeWidget.Pin()->GetDesiredSize().X / 2 + PosX, PosY), NodeFilter);
+					AutoArrangeNodes(GraphNode, BBoxTree.Children[BBoxIndex], PosX, PosY + GraphNode->DEPRECATED_NodeWidget.Pin()->GetDesiredSize().Y * 2.5f);
+					GraphNode->DEPRECATED_NodeWidget.Pin()->MoveTo(FVector2D(BBoxTree.Children[BBoxIndex].SubGraphBBox.X / 2 - GraphNode->DEPRECATED_NodeWidget.Pin()->GetDesiredSize().X / 2 + PosX, PosY), NodeFilter);
 					PosX += BBoxTree.Children[BBoxIndex].SubGraphBBox.X + 20;
 				}
 
@@ -1095,7 +1095,7 @@ namespace BTAutoArrangeHelpers
 
 	void GetNodeSizeInfo(UBehaviorTreeGraphNode* ParentNode, FNodeBoundsInfo& BBoxTree)
 	{
-		BBoxTree.SubGraphBBox = ParentNode->NodeWidget.Pin()->GetDesiredSize();
+		BBoxTree.SubGraphBBox = ParentNode->DEPRECATED_NodeWidget.Pin()->GetDesiredSize();
 		float LevelWidth = 0;
 		float LevelHeight = 0;
 
@@ -1145,12 +1145,12 @@ void UBehaviorTreeGraph::AutoArrange()
 
 	BTAutoArrangeHelpers::FNodeBoundsInfo BBoxTree;
 	BTAutoArrangeHelpers::GetNodeSizeInfo(RootNode, BBoxTree);
-	BTAutoArrangeHelpers::AutoArrangeNodes(RootNode, BBoxTree, 0, RootNode->NodeWidget.Pin()->GetDesiredSize().Y * 2.5f);
+	BTAutoArrangeHelpers::AutoArrangeNodes(RootNode, BBoxTree, 0, RootNode->DEPRECATED_NodeWidget.Pin()->GetDesiredSize().Y * 2.5f);
 
-	RootNode->NodePosX = BBoxTree.SubGraphBBox.X / 2 - RootNode->NodeWidget.Pin()->GetDesiredSize().X / 2;
+	RootNode->NodePosX = BBoxTree.SubGraphBBox.X / 2 - RootNode->DEPRECATED_NodeWidget.Pin()->GetDesiredSize().X / 2;
 	RootNode->NodePosY = 0;
 
-	RootNode->NodeWidget.Pin()->GetOwnerPanel()->ZoomToFit(/*bOnlySelection=*/ false);
+	RootNode->DEPRECATED_NodeWidget.Pin()->GetOwnerPanel()->ZoomToFit(/*bOnlySelection=*/ false);
 }
 
 void UBehaviorTreeGraph::OnSubNodeDropped()
diff --git a/Engine/Source/Editor/GraphEditor/Private/DragConnection.cpp b/Engine/Source/Editor/GraphEditor/Private/DragConnection.cpp
--- a/Engine/Source/Editor/GraphEditor/Private/DragConnection.cpp
+++ b/Engine/Source/Editor/GraphEditor/Private/DragConnection.cpp
@@ -52,7 +52,7 @@ void FDragConnection::HoverTargetChanged()
 
 				if (Response.Response == ECanCreateConnectionResponse::CONNECT_RESPONSE_DISALLOW)
 				{
-					 TSharedPtr<SGraphNode> NodeWidget = TargetPinObj->GetOwningNode()->NodeWidget.Pin();
+					TSharedPtr<SGraphNode> NodeWidget = TargetPinObj->GetOwningNode()->DEPRECATED_NodeWidget.Pin();
 					 if (NodeWidget.IsValid())
 					 {
 						 NodeWidget->NotifyDisallowedPinConnection(StartingPinObj, TargetPinObj);
diff --git a/Engine/Source/Editor/GraphEditor/Private/SCommentBubble.cpp b/Engine/Source/Editor/GraphEditor/Private/SCommentBubble.cpp
--- a/Engine/Source/Editor/GraphEditor/Private/SCommentBubble.cpp
+++ b/Engine/Source/Editor/GraphEditor/Private/SCommentBubble.cpp
@@ -320,9 +320,9 @@ float SCommentBubble::GetArrowCenterOffset() const
 {
 	float CentreOffset = GraphNode->bCommentBubbleVisible ? SCommentBubbleDefs::ArrowCentreOffset : SCommentBubbleDefs::ToggleButtonCentreOffset;
 	const bool bVisibleAndUnpinned = !GraphNode->bCommentBubblePinned && GraphNode->bCommentBubbleVisible;
-	if( bVisibleAndUnpinned && GraphNode->NodeWidget.IsValid() )
+	if (bVisibleAndUnpinned && GraphNode->DEPRECATED_NodeWidget.IsValid())
 	{
-		TSharedPtr<SGraphPanel> GraphPanel = GraphNode->NodeWidget.Pin()->GetOwnerPanel();
+		TSharedPtr<SGraphPanel> GraphPanel = GraphNode->DEPRECATED_NodeWidget.Pin()->GetOwnerPanel();
 		CentreOffset *= GraphPanel.IsValid() ? GraphPanel->GetZoomAmount() : 1.f;
 	}
 	return CentreOffset;
@@ -463,9 +463,9 @@ void SCommentBubble::OnPinStateToggle( ECheckBoxState State ) const
 bool SCommentBubble::IsReadOnly() const
 {
 	bool bReadOnly = true;
-	if( GraphNode && GraphNode->NodeWidget.IsValid() )
+	if (GraphNode && GraphNode->DEPRECATED_NodeWidget.IsValid())
 	{
-		bReadOnly = !GraphNode->NodeWidget.Pin()->IsNodeEditable();
+		bReadOnly = !GraphNode->DEPRECATED_NodeWidget.Pin()->IsNodeEditable();
 	}
 	return bReadOnly;
 }
\ No newline at end of file
diff --git a/Engine/Source/Editor/GraphEditor/Private/SGraphNode.cpp b/Engine/Source/Editor/GraphEditor/Private/SGraphNode.cpp
--- a/Engine/Source/Editor/GraphEditor/Private/SGraphNode.cpp
+++ b/Engine/Source/Editor/GraphEditor/Private/SGraphNode.cpp
@@ -449,7 +449,7 @@ void SGraphNode::SetOwner( const TSharedRef<SGraphPanel>& OwnerPanel )
 	check( !OwnerGraphPanelPtr.IsValid() );
 	SetParentPanel(OwnerPanel);
 	OwnerGraphPanelPtr = OwnerPanel;
-	GraphNode->NodeWidget = SharedThis(this);
+	GraphNode->DEPRECATED_NodeWidget = SharedThis(this);
 
 	/*Once we have an owner, and if hide Unused pins is enabled, we need to remake our pins to drop the hidden ones*/
 	if(OwnerGraphPanelPtr.Pin()->GetPinVisibility() != SGraphEditor::Pin_Show 
diff --git a/Engine/Source/Editor/GraphEditor/Private/SGraphNodeComment.cpp b/Engine/Source/Editor/GraphEditor/Private/SGraphNodeComment.cpp
--- a/Engine/Source/Editor/GraphEditor/Private/SGraphNodeComment.cpp
+++ b/Engine/Source/Editor/GraphEditor/Private/SGraphNodeComment.cpp
@@ -356,9 +356,9 @@ void SGraphNodeComment::MoveTo( const FVector2D& NewPosition, FNodeSet& NodeFilt
 			{
 				if (UEdGraphNode* Node = Cast<UEdGraphNode>(*NodeIt))
 				{
-					if ( !Panel->SelectionManager.IsNodeSelected(Node) && !NodeFilter.Find( Node->NodeWidget.Pin() ))
+					if (!Panel->SelectionManager.IsNodeSelected(Node) && !NodeFilter.Find(Node->DEPRECATED_NodeWidget.Pin()))
 					{
-						NodeFilter.Add(Node->NodeWidget.Pin());
+						NodeFilter.Add(Node->DEPRECATED_NodeWidget.Pin());
 						Node->Modify();
 						Node->NodePosX += PositionDelta.X;
 						Node->NodePosY += PositionDelta.Y;
diff --git a/Engine/Source/Editor/Kismet/Private/BlueprintEditor.cpp b/Engine/Source/Editor/Kismet/Private/BlueprintEditor.cpp
--- a/Engine/Source/Editor/Kismet/Private/BlueprintEditor.cpp
+++ b/Engine/Source/Editor/Kismet/Private/BlueprintEditor.cpp
@@ -2970,7 +2970,7 @@ void FBlueprintEditor::OnBlueprintCompiled(UBlueprint* InBlueprint)
 			{
 				if (Node)
 				{
-					auto Widget = Node->NodeWidget.Pin();
+					auto Widget = Node->DEPRECATED_NodeWidget.Pin();
 					if (Widget.IsValid())
 					{
 						Widget->RefreshErrorInfo();
@@ -3968,7 +3968,7 @@ void FBlueprintEditor::OnAddParentNode()
 				int32 NodeSizeY = 15;
 				if( UK2Node* Node = Cast<UK2Node>(SelectedObj))
 				{
-					NodeSizeY += Node->NodeWidget.IsValid() ? static_cast<int32>(Node->NodeWidget.Pin()->GetDesiredSize().Y) : 0;
+					NodeSizeY += Node->DEPRECATED_NodeWidget.IsValid() ? static_cast<int32>(Node->DEPRECATED_NodeWidget.Pin()->GetDesiredSize().Y) : 0;
 				}
 				ParentFunctionNode->NodePosX = FunctionFromNode.Node->NodePosX;
 				ParentFunctionNode->NodePosY = FunctionFromNode.Node->NodePosY + NodeSizeY;
diff --git a/Engine/Source/Runtime/Engine/Classes/EdGraph/EdGraphNode.h b/Engine/Source/Runtime/Engine/Classes/EdGraph/EdGraphNode.h
--- a/Engine/Source/Runtime/Engine/Classes/EdGraph/EdGraphNode.h
+++ b/Engine/Source/Runtime/Engine/Classes/EdGraph/EdGraphNode.h
@@ -167,8 +167,8 @@ public:
 	virtual void PostLoad() override;
 	// End of UObject interface
 
-	/** widget representing this node if exists */
-	TWeakPtr<SGraphNode> NodeWidget;
+	/** widget representing this node if it exists; Note: This is not safe to use in general and will be removed in the future, as there is no guarantee that only one graph editor/panel is viewing a given graph */
+	TWeakPtr<SGraphNode> DEPRECATED_NodeWidget;
 
 	/** Create a new pin on this node using the supplied info, and return the new pin */
 	UEdGraphPin* CreatePin(EEdGraphPinDirection Dir, const FString& PinCategory, const FString& PinSubCategory, UObject* PinSubCategoryObject, bool bIsArray, bool bIsReference, const FString& PinName, bool bIsConst = false, int32 Index = INDEX_NONE);
