Changed two checks in Blackboard Editor to error messages #UE4

UE-20216

diff --git a/Engine/Source/Editor/BehaviorTreeEditor/Private/SBehaviorTreeBlackboardEditor.cpp b/Engine/Source/Editor/BehaviorTreeEditor/Private/SBehaviorTreeBlackboardEditor.cpp
--- a/Engine/Source/Editor/BehaviorTreeEditor/Private/SBehaviorTreeBlackboardEditor.cpp
+++ b/Engine/Source/Editor/BehaviorTreeEditor/Private/SBehaviorTreeBlackboardEditor.cpp
@@ -17,6 +17,8 @@
 
 #define LOCTEXT_NAMESPACE "SBehaviorTreeBlackboardEditor"
 
+DEFINE_LOG_CATEGORY_STATIC(LogBlackboardEditor, Warning, All);
+
 void SBehaviorTreeBlackboardEditor::Construct(const FArguments& InArgs, TSharedRef<FUICommandList> InCommandList, UBlackboardData* InBlackboardData)
 {
 	OnEntrySelected = InArgs._OnEntrySelected;
@@ -91,7 +93,11 @@ TSharedPtr<FExtender> SBehaviorTreeBlackboardEditor::GetToolbarExtender(TSharedR
 
 void SBehaviorTreeBlackboardEditor::HandleDeleteEntry()
 {
-	check(BlackboardData);
+	if (BlackboardData == nullptr)
+	{
+		UE_LOG(LogBlackboardEditor, Error, TEXT("Trying to delete an entry from a blackboard while no Blackboard Asset is set!"));
+		return;
+	}
 
 	if(!IsDebuggerActive())
 	{
@@ -208,7 +214,11 @@ TSharedRef<SWidget> SBehaviorTreeBlackboardEditor::HandleCreateNewEntryMenu() co
 
 void SBehaviorTreeBlackboardEditor::HandleKeyClassPicked(UClass* InClass)
 {
-	check(BlackboardData);
+	if (BlackboardData == nullptr)
+	{
+		UE_LOG(LogBlackboardEditor, Error, TEXT("Trying to delete an entry from a blackboard while no Blackboard Asset is set!"));
+		return;
+	}
 
 	FSlateApplication::Get().DismissAllMenus();
 
