Fix behaviortree debugger not retreiving the instance data from the right world when opening it after PIE is started.

diff --git a/Engine/Source/Editor/BehaviorTreeEditor/Private/BehaviorTreeDebugger.cpp b/Engine/Source/Editor/BehaviorTreeEditor/Private/BehaviorTreeDebugger.cpp
--- a/Engine/Source/Editor/BehaviorTreeEditor/Private/BehaviorTreeDebugger.cpp
+++ b/Engine/Source/Editor/BehaviorTreeEditor/Private/BehaviorTreeDebugger.cpp
@@ -738,13 +738,33 @@ void FBehaviorTreeDebugger::FindLockedDebugActor(UWorld* World)
 void FBehaviorTreeDebugger::FindMatchingTreeInstance()
 {
 	KnownInstances.Reset();
-	if (GEditor->PlayWorld == NULL)
+
+    // Find the world for the dedicated server if any, otherwise fallback to the PIE world
+	UWorld* PlayWorld = nullptr;
+	for (const FWorldContext& PieContext : GEditor->GetWorldContexts())
+	{
+		if (PieContext.WorldType == EWorldType::PIE && PieContext.World() != nullptr)
+		{
+			if (PieContext.RunAsDedicated)
+			{
+				PlayWorld = PieContext.World();
+				break;
+			}
+			else if(!PlayWorld)
+			{
+				PlayWorld = PieContext.World();
+				// Need to continue to see if their is a dedicated server.
+			}
+		}
+	}
+	 
+	if (PlayWorld == NULL)
 	{
 		return;
 	}
 
 	UBehaviorTreeComponent* MatchingComp = NULL;
-	for (FActorIterator It(GEditor->PlayWorld); It; ++It)
+	for (FActorIterator It(PlayWorld); It; ++It)
 	{
 		AActor* TestActor = *It;
 		UBehaviorTreeComponent* TestComp = TestActor ? TestActor->FindComponentByClass<UBehaviorTreeComponent>() : nullptr;
