Fix BT not pausing dedicated server world upon breakpoint.

[RN] Behavior Tree now pauses all game in PIE when running with client and dedicated sever.

diff --git a/Engine/Source/Editor/BehaviorTreeEditor/Private/BehaviorTreeDebugger.cpp b/Engine/Source/Editor/BehaviorTreeEditor/Private/BehaviorTreeDebugger.cpp
--- a/Engine/Source/Editor/BehaviorTreeEditor/Private/BehaviorTreeDebugger.cpp
+++ b/Engine/Source/Editor/BehaviorTreeEditor/Private/BehaviorTreeDebugger.cpp
@@ -835,9 +835,31 @@ void FBehaviorTreeDebugger::StepForwardInto()
 #endif
 }
 
+void ForEachGameWorld(const TFunction<void(UWorld*)>& Func)
+{
+	for (const FWorldContext& PieContext : GUnrealEd->GetWorldContexts())
+	{
+		UWorld* PlayWorld = PieContext.World();
+		if (PlayWorld && PlayWorld->IsGameWorld())
+		{
+			Func(PlayWorld);
+		}
+	}
+}
+
+bool AreAllGameWorldPaused()
+{
+	bool bPaused = true;
+	ForEachGameWorld([&](UWorld* World)
+	{ 
+		bPaused = bPaused && World->bDebugPauseExecution; 
+	});
+	return bPaused;
+}
+
 bool FBehaviorTreeDebugger::CanStepForwardInto() const
 {
-	return GUnrealEd->PlayWorld && GUnrealEd->PlayWorld->bDebugPauseExecution && (StepForwardIntoIdx != INDEX_NONE);
+	return AreAllGameWorldPaused() && (StepForwardIntoIdx != INDEX_NONE);
 }
 
 void FBehaviorTreeDebugger::StepForwardOver()
@@ -849,7 +871,7 @@ void FBehaviorTreeDebugger::StepForwardOver()
 
 bool FBehaviorTreeDebugger::CanStepForwardOver() const
 {
-	return GUnrealEd->PlayWorld && GUnrealEd->PlayWorld->bDebugPauseExecution && (StepForwardOverIdx != INDEX_NONE);
+	return AreAllGameWorldPaused() && (StepForwardOverIdx != INDEX_NONE);
 }
 
 void FBehaviorTreeDebugger::StepOut()
@@ -861,7 +883,7 @@ void FBehaviorTreeDebugger::StepOut()
 
 bool FBehaviorTreeDebugger::CanStepOut() const
 {
-	return GUnrealEd->PlayWorld && GUnrealEd->PlayWorld->bDebugPauseExecution && (StepOutIdx != INDEX_NONE);
+	return AreAllGameWorldPaused() && (StepOutIdx != INDEX_NONE);
 }
 
 void FBehaviorTreeDebugger::StepBackInto()
@@ -873,7 +895,7 @@ void FBehaviorTreeDebugger::StepBackInto()
 
 bool FBehaviorTreeDebugger::CanStepBackInto() const
 {
-	return GUnrealEd->PlayWorld && GUnrealEd->PlayWorld->bDebugPauseExecution && (StepBackIntoIdx != INDEX_NONE);
+	return AreAllGameWorldPaused() && (StepBackIntoIdx != INDEX_NONE);
 }
 
 void FBehaviorTreeDebugger::StepBackOver()
@@ -885,7 +907,7 @@ void FBehaviorTreeDebugger::StepBackOver()
 
 bool FBehaviorTreeDebugger::CanStepBackOver() const
 {
-	return GUnrealEd->PlayWorld && GUnrealEd->PlayWorld->bDebugPauseExecution && (StepBackOverIdx != INDEX_NONE);
+	return AreAllGameWorldPaused() && (StepBackOverIdx != INDEX_NONE);
 }
 
 void FBehaviorTreeDebugger::UpdateCurrentStep(int32 PrevStepIdx, int32 NewStepIdx)
@@ -1010,30 +1032,46 @@ void FBehaviorTreeDebugger::StopPlaySession()
 
 void FBehaviorTreeDebugger::PausePlaySession()
 {
-	if (GUnrealEd->PlayWorld && !GUnrealEd->PlayWorld->bDebugPauseExecution)
+	bool bPaused = false;
+	ForEachGameWorld([&](UWorld* World)
+	{
+		if (!World->bDebugPauseExecution)
+		{
+			World->bDebugPauseExecution = true;
+			bPaused = true;
+		}
+	});
+	if (bPaused)
 	{
-		GUnrealEd->PlayWorld->bDebugPauseExecution = true;
 		GUnrealEd->PlaySessionPaused();
 	}
 }
 
 void FBehaviorTreeDebugger::ResumePlaySession()
 {
-	if (GUnrealEd->PlayWorld && GUnrealEd->PlayWorld->bDebugPauseExecution)
+	bool bResumed = false;
+	ForEachGameWorld([&](UWorld* World)
+	{
+		if (World->bDebugPauseExecution)
+		{
+			World->bDebugPauseExecution = false;
+			bResumed = true;
+		}
+	});
+	if(bResumed)
 	{
-		GUnrealEd->PlayWorld->bDebugPauseExecution = false;
 		GUnrealEd->PlaySessionResumed();
 	}
 }
 
 bool FBehaviorTreeDebugger::IsPlaySessionPaused()
 {
-	return GUnrealEd->PlayWorld && GUnrealEd->PlayWorld->bDebugPauseExecution;
+	return AreAllGameWorldPaused();
 }
 
 bool FBehaviorTreeDebugger::IsPlaySessionRunning()
 {
-	return GUnrealEd->PlayWorld && !GUnrealEd->PlayWorld->bDebugPauseExecution;
+	return !AreAllGameWorldPaused();
 }
 
 bool FBehaviorTreeDebugger::IsPIESimulating()
