Blackboard key renaming now loads only the BehaviorTree assets referencing the parent BlackboardData asset.

diff --git a/Engine/Source/Editor/BehaviorTreeEditor/Private/SBehaviorTreeBlackboardView.cpp b/Engine/Source/Editor/BehaviorTreeEditor/Private/SBehaviorTreeBlackboardView.cpp
--- a/Engine/Source/Editor/BehaviorTreeEditor/Private/SBehaviorTreeBlackboardView.cpp
+++ b/Engine/Source/Editor/BehaviorTreeEditor/Private/SBehaviorTreeBlackboardView.cpp
@@ -202,12 +202,13 @@ private:
 		FName NewName = FName(*AsString);
 		if(NewName != OldName)
 		{
-			if(!BlackboardEntryAction->bIsNew)
+			TArray<UObject*> ExternalBTAssetsWithKeyReferences;
+			if(!BlackboardEntryAction->bIsNew && BlackboardEntryAction->BlackboardData)
 			{
 				// Preload behavior trees before we transact otherwise they will add objects to 
 				// the transaction buffer whether we change them or not.
 				// Blueprint regeneration does this in UEdGraphNode::CreatePin.
-				LoadAllBehaviorTrees();
+				LoadReferencerBehaviorTrees(*(BlackboardEntryAction->BlackboardData), ExternalBTAssetsWithKeyReferences);
 			}
 
 			const FScopedTransaction Transaction(LOCTEXT("BlackboardEntryRenameTransaction", "Rename Blackboard Entry"));
@@ -221,85 +222,68 @@ private:
 
 			if(!BlackboardEntryAction->bIsNew)
 			{
-				UpdateExternalBlackboardKeyReferences(OldName, NewName);
+				UpdateExternalBlackboardKeyReferences(OldName, NewName, ExternalBTAssetsWithKeyReferences);
 			}
 		}
 
 		BlackboardEntryAction->bIsNew = false;
 	}
 
-	void LoadAllBehaviorTrees()
+	void LoadReferencerBehaviorTrees(const UBlackboardData& InBlackboardData, TArray<UObject*>& OutExternalBTAssetsWithKeyReferences)
 	{
-		FAssetRegistryModule& AssetRegistry = FModuleManager::LoadModuleChecked<FAssetRegistryModule>(TEXT("AssetRegistry"));
+		IAssetRegistry& AssetRegistry = FModuleManager::LoadModuleChecked<FAssetRegistryModule>(TEXT("AssetRegistry")).Get();
 
-		FARFilter Filter;
-		Filter.ClassNames.Add(UBehaviorTree::StaticClass()->GetFName());
-		Filter.bRecursiveClasses = true;
+		TArray<FName> ReferencerPackages;
+		AssetRegistry.GetReferencers(InBlackboardData.GetOutermost()->GetFName(), ReferencerPackages, EAssetRegistryDependencyType::Hard);
 
-		TArray<FAssetData> AssetData;
-		AssetRegistry.Get().GetAssets(Filter, AssetData);
+		if (ReferencerPackages.Num())
+		{
+			FScopedSlowTask SlowTask((float)ReferencerPackages.Num(), LOCTEXT("UpdatingBehaviorTrees", "Updating behavior trees"));
+			SlowTask.MakeDialog();
 
-		FScopedSlowTask SlowTask((float)AssetData.Num(), LOCTEXT("UpdatingBehaviorTrees", "Updating behavior trees"));
-		SlowTask.MakeDialog();
+			for (const FName& ReferencerPackage : ReferencerPackages)
+			{
+				TArray<FAssetData> Assets;
+				AssetRegistry.GetAssetsByPackageName(ReferencerPackage, Assets);
 
-		for (const auto& BehaviorTreeAsset : AssetData)
-		{
-			SlowTask.EnterProgressFrame(1.0f, FText::Format(LOCTEXT("CheckingBehaviorTree", "Key renamed, loading {0}"), FText::FromName(BehaviorTreeAsset.AssetName)));
-			BehaviorTreeAsset.GetAsset();
+				for (const FAssetData& Asset : Assets)
+				{
+					if (Asset.GetClass()->IsChildOf(UBehaviorTree::StaticClass()))
+					{
+						SlowTask.EnterProgressFrame(1.0f, FText::Format(LOCTEXT("CheckingBehaviorTree", "Key renamed, loading {0}"), FText::FromName(Asset.AssetName)));
+
+						UBehaviorTree* BehaviorTree = Cast<UBehaviorTree>(Asset.GetAsset());
+						if (BehaviorTree->BlackboardAsset == &InBlackboardData)
+						{
+							OutExternalBTAssetsWithKeyReferences.Add(BehaviorTree);
+						}
+					}
+				}
+			}
 		}
 	}
 
 	#define GET_STRUCT_NAME_CHECKED(StructName) \
 		((void)sizeof(StructName), TEXT(#StructName))
 
-	void UpdateExternalBlackboardKeyReferences(const FName& OldKey, const FName& NewKey) const
+	void UpdateExternalBlackboardKeyReferences(const FName& OldKey, const FName& NewKey, const TArray<UObject*>& InExternalBTAssetsWithKeyReferences) const
 	{
-		// update all behavior trees that reference this key
-		check(ActionPtr.Pin()->GetTypeId() == FEdGraphSchemaAction_BlackboardEntry::StaticGetTypeId());
-		TSharedPtr<FEdGraphSchemaAction_BlackboardEntry> BlackboardEntryAction = StaticCastSharedPtr<FEdGraphSchemaAction_BlackboardEntry>(ActionPtr.Pin());
-
-		FAssetRegistryModule& AssetRegistry = FModuleManager::LoadModuleChecked<FAssetRegistryModule>(TEXT("AssetRegistry"));
-
-		FARFilter Filter;
-		Filter.ClassNames.Add(UBehaviorTree::StaticClass()->GetFName());
-		Filter.bRecursiveClasses = true;
-
-		TArray<FAssetData> AssetData;
-		AssetRegistry.Get().GetAssets(Filter, AssetData);
-
-		for (const auto& BehaviorTreeAsset : AssetData)
+		for (const UObject* Asset : InExternalBTAssetsWithKeyReferences)
 		{
-			UBehaviorTree* BehaviorTree = Cast<UBehaviorTree>(BehaviorTreeAsset.GetAsset());
-			if(BehaviorTree)
+			// search all subobjects of this package for FBlackboardKeySelector structs and update as necessary
+			TArray<UObject*> Objects;
+			GetObjectsWithOuter(Asset->GetOutermost(), Objects);
+			for (const auto& SubObject : Objects)
 			{
-				if(BehaviorTree->BlackboardAsset == BlackboardEntryAction->BlackboardData)
+				for (TFieldIterator<UStructProperty> It(SubObject->GetClass()); It; ++It)
 				{
-					UPackage* BehaviorTreePackage = BehaviorTreeAsset.GetPackage();
-					if(BehaviorTreePackage)
+					if (It->GetCPPType(NULL, CPPF_None).Contains(GET_STRUCT_NAME_CHECKED(FBlackboardKeySelector)))
 					{
-						// search all subobjects of this package for FBlackboardKeySelector structs and update as necessary
-						TArray<UObject*> Objects;
-						GetObjectsWithOuter(BehaviorTreePackage, Objects);
-						for(const auto& SubObject : Objects)
+						FBlackboardKeySelector* PropertyValue = (FBlackboardKeySelector*)(It->ContainerPtrToValuePtr<uint8>(SubObject));
+						if (PropertyValue && PropertyValue->SelectedKeyName == OldKey)
 						{
-							for (UProperty* Property = SubObject->GetClass()->PropertyLink; Property; Property = Property->PropertyLinkNext)
-							{
-								uint8* PropertyData = Property->ContainerPtrToValuePtr<uint8>(SubObject);
-								UStructProperty* StructProperty = Cast<UStructProperty>(Property);
-
-								if (StructProperty && StructProperty->GetCPPType(NULL, CPPF_None).Contains(GET_STRUCT_NAME_CHECKED(FBlackboardKeySelector)))
-								{
-									FBlackboardKeySelector* PropertyValue = (FBlackboardKeySelector*)PropertyData;
-									if(PropertyValue)
-									{
-										if(PropertyValue->SelectedKeyName == OldKey)
-										{
-											SubObject->Modify();
-											PropertyValue->SelectedKeyName = NewKey;
-										}
-									}
-								}
-							}
+							SubObject->Modify();
+							PropertyValue->SelectedKeyName = NewKey;
 						}
 					}
 				}
