Added option "DisplayBlackboardKeysInAlphabeticalOrder" in UEditorPerProjectUserSettings to display blackboard keys in BlackboardView, BlackboardEditor and BlackboardKeySelector in alphabetical order.
Bonus: autoexpand blackboard keys

[RN] Added option DisplayBlackboardKeysInAlphabeticalOrder in UEditorPerProjectUserSettings to display blackboard keys in BlackboardView, BlackboardEditor and BlackboardKeySelector in alphabetical order.

diff --git a/Engine/Source/Editor/BehaviorTreeEditor/Private/DetailCustomizations/BlackboardSelectorDetails.cpp b/Engine/Source/Editor/BehaviorTreeEditor/Private/DetailCustomizations/BlackboardSelectorDetails.cpp
--- a/Engine/Source/Editor/BehaviorTreeEditor/Private/DetailCustomizations/BlackboardSelectorDetails.cpp
+++ b/Engine/Source/Editor/BehaviorTreeEditor/Private/DetailCustomizations/BlackboardSelectorDetails.cpp
@@ -15,6 +15,7 @@
 #include "BehaviorTree/Blackboard/BlackboardKeyType.h"
 #include "BehaviorTree/BlackboardData.h"
 #include "IPropertyUtilities.h"
+#include "Editor/EditorPerProjectUserSettings.h"
 
 #define LOCTEXT_NAMESPACE "BlackboardSelectorDetails"
 
@@ -148,6 +149,11 @@ void FBlackboardSelectorDetails::CacheBlackboardData()
 			break;
 		}
 	}
+
+	if (GetDefault<UEditorPerProjectUserSettings>()->bDisplayBlackboardKeysInAlphabeticalOrder)
+	{
+		KeyValues.Sort([](const FName& a, const FName& b) { return a < b; });
+	}	
 }
 
 void FBlackboardSelectorDetails::InitKeyFromProperty()
diff --git a/Engine/Source/Editor/BehaviorTreeEditor/Private/SBehaviorTreeBlackboardView.cpp b/Engine/Source/Editor/BehaviorTreeEditor/Private/SBehaviorTreeBlackboardView.cpp
--- a/Engine/Source/Editor/BehaviorTreeEditor/Private/SBehaviorTreeBlackboardView.cpp
+++ b/Engine/Source/Editor/BehaviorTreeEditor/Private/SBehaviorTreeBlackboardView.cpp
@@ -18,6 +18,7 @@
 #include "Widgets/Text/STextBlock.h"
 #include "Framework/MultiBox/MultiBoxDefs.h"
 #include "Framework/MultiBox/MultiBoxBuilder.h"
+#include "Editor/EditorPerProjectUserSettings.h"
 #include "EditorStyleSet.h"
 #include "SGraphActionMenu.h"
 #include "SGraphPalette.h"
@@ -472,7 +473,8 @@ void SBehaviorTreeBlackboardView::Construct(const FArguments& InArgs, TSharedRef
 				.OnActionSelected(this, &SBehaviorTreeBlackboardView::HandleActionSelected)
 				.OnContextMenuOpening(this, &SBehaviorTreeBlackboardView::HandleContextMenuOpening, InCommandList)
 				.OnActionMatchesName(this, &SBehaviorTreeBlackboardView::HandleActionMatchesName)
-				.AlphaSortItems(false)
+				.AlphaSortItems(GetDefault<UEditorPerProjectUserSettings>()->bDisplayBlackboardKeysInAlphabeticalOrder)
+				.AutoExpandActionMenu(true)
 			]
 		]
 	];
diff --git a/Engine/Source/Editor/UnrealEd/Classes/Editor/EditorPerProjectUserSettings.h b/Engine/Source/Editor/UnrealEd/Classes/Editor/EditorPerProjectUserSettings.h
--- a/Engine/Source/Editor/UnrealEd/Classes/Editor/EditorPerProjectUserSettings.h
+++ b/Engine/Source/Editor/UnrealEd/Classes/Editor/EditorPerProjectUserSettings.h
@@ -32,6 +32,10 @@ class UEditorPerProjectUserSettings : public UObject
 	UPROPERTY(EditAnywhere, config, Category = AI)
 	uint32 bAlwaysGatherBehaviorTreeDebuggerData : 1;
 
+	/** If enabled, blackboard keys displayed in blackboard editor and key selector will be sorted in alphabetical order . */
+	UPROPERTY(EditAnywhere, config, Category = AI)
+	uint32 bDisplayBlackboardKeysInAlphabeticalOrder : 1;
+
 	/** When enabled, Engine Version Number is displayed in the ProjectBadge */
 	UPROPERTY(EditAnywhere, config, Category = DeveloperTools, meta = (DisplayName = "Display Engine Version Number in Project Badge", ConfigRestartRequired = true))
 	bool bDisplayEngineVersionInBadge;
diff --git a/Engine/Source/Editor/UnrealEd/Private/Settings/EditorPerProjectUserSettings.cpp b/Engine/Source/Editor/UnrealEd/Private/Settings/EditorPerProjectUserSettings.cpp
--- a/Engine/Source/Editor/UnrealEd/Private/Settings/EditorPerProjectUserSettings.cpp
+++ b/Engine/Source/Editor/UnrealEd/Private/Settings/EditorPerProjectUserSettings.cpp
@@ -20,6 +20,8 @@ UEditorPerProjectUserSettings::UEditorPerProjectUserSettings(const FObjectInitia
 	AssetViewerProfileIndex = 0;
 	bAnimationReimportWarnings = false;
 
+	bDisplayBlackboardKeysInAlphabeticalOrder = false;
+
 	SimplygonServerIP = "127.0.0.1";
 	SimplygonSwarmDelay = 5000;
 	bEnableSwarmDebugging = false;
