Replaced deprecated 'FName::operator <' by  LexicalLess() for Blackboard key sorting

#rnx

diff --git a/Engine/Source/Editor/BehaviorTreeEditor/Private/DetailCustomizations/BlackboardSelectorDetails.cpp b/Engine/Source/Editor/BehaviorTreeEditor/Private/DetailCustomizations/BlackboardSelectorDetails.cpp
--- a/Engine/Source/Editor/BehaviorTreeEditor/Private/DetailCustomizations/BlackboardSelectorDetails.cpp
+++ b/Engine/Source/Editor/BehaviorTreeEditor/Private/DetailCustomizations/BlackboardSelectorDetails.cpp
@@ -152,7 +152,7 @@ void FBlackboardSelectorDetails::CacheBlackboardData()
 
 	if (GetDefault<UEditorPerProjectUserSettings>()->bDisplayBlackboardKeysInAlphabeticalOrder)
 	{
-		KeyValues.Sort([](const FName& a, const FName& b) { return a < b; });
+		KeyValues.Sort([](const FName& a, const FName& b) { return a.LexicalLess(b); });
 	}	
 }
 
