PR #5695: Create comment keyboard shortcut for Behavior tree graphs (Contributed by [REDACTED_PERSON])

diff --git a/Engine/Source/Editor/BehaviorTreeEditor/Private/BehaviorTreeEditor.cpp b/Engine/Source/Editor/BehaviorTreeEditor/Private/BehaviorTreeEditor.cpp
--- a/Engine/Source/Editor/BehaviorTreeEditor/Private/BehaviorTreeEditor.cpp
+++ b/Engine/Source/Editor/BehaviorTreeEditor/Private/BehaviorTreeEditor.cpp
@@ -567,6 +567,12 @@ TSharedRef<SGraphEditor> FBehaviorTreeEditor::CreateGraphEditorWidget(UEdGraph*
 			FIsActionChecked(),
 			FIsActionButtonVisible::CreateSP( this, &FBehaviorTreeEditor::CanToggleBreakpoint )
 			);
+
+		GraphEditorCommands->MapAction(
+			FGraphEditorCommands::Get().CreateComment,
+			FExecuteAction::CreateSP(this, &FBehaviorTreeEditor::OnCreateComment),
+			FCanExecuteAction::CreateSP(this, &FBehaviorTreeEditor::CanCreateComment)
+			);
 	}
 
 	SGraphEditor::FGraphEditorEvents InEvents;
@@ -1840,4 +1846,21 @@ bool FBehaviorTreeEditor::CanCreateNewBlackboard() const
 	return !IsDebuggerReady();
 }
 
+bool FBehaviorTreeEditor::CanCreateComment() const
+{
+	return (SelectedNodesCount > 1);
+}
+
+void FBehaviorTreeEditor::OnCreateComment()
+{
+	if (BehaviorTree && BehaviorTree->BTGraph)
+	{
+		TSharedPtr<FEdGraphSchemaAction> Action = BehaviorTree->BTGraph->GetSchema()->GetCreateCommentAction();
+		if (Action.IsValid())
+		{
+			Action->PerformAction(BehaviorTree->BTGraph, nullptr, FVector2D());
+		}
+	}
+}
+
 #undef LOCTEXT_NAMESPACE
diff --git a/Engine/Source/Editor/BehaviorTreeEditor/Private/BehaviorTreeEditor.h b/Engine/Source/Editor/BehaviorTreeEditor/Private/BehaviorTreeEditor.h
--- a/Engine/Source/Editor/BehaviorTreeEditor/Private/BehaviorTreeEditor.h
+++ b/Engine/Source/Editor/BehaviorTreeEditor/Private/BehaviorTreeEditor.h
@@ -85,6 +85,8 @@ public:
 	bool CanAddBreakpoint() const;
 	void OnRemoveBreakpoint();
 	bool CanRemoveBreakpoint() const;
+	void OnCreateComment();
+	bool CanCreateComment() const;
 
 	void SearchTree();
 	bool CanSearchTree() const;
