Made BT editor comments creation consistent with BP editor so that comment gets created event when a single node is  selected #UE4

diff --git a/Engine/Source/Editor/BehaviorTreeEditor/Private/BehaviorTreeEditor.cpp b/Engine/Source/Editor/BehaviorTreeEditor/Private/BehaviorTreeEditor.cpp
--- a/Engine/Source/Editor/BehaviorTreeEditor/Private/BehaviorTreeEditor.cpp
+++ b/Engine/Source/Editor/BehaviorTreeEditor/Private/BehaviorTreeEditor.cpp
@@ -1848,7 +1848,7 @@ bool FBehaviorTreeEditor::CanCreateNewBlackboard() const
 
 bool FBehaviorTreeEditor::CanCreateComment() const
 {
-	return (SelectedNodesCount > 1);
+	return (SelectedNodesCount > 0);
 }
 
 void FBehaviorTreeEditor::OnCreateComment()
