PR #5720: Added BEHAVIORTREEEDITOR_API to "BehaviorTreeGraphNode_Composite" (Contributed by [REDACTED_PERSON])

diff --git a/Engine/Source/Editor/BehaviorTreeEditor/Classes/BehaviorTreeGraphNode_Composite.h b/Engine/Source/Editor/BehaviorTreeEditor/Classes/BehaviorTreeGraphNode_Composite.h
--- a/Engine/Source/Editor/BehaviorTreeEditor/Classes/BehaviorTreeGraphNode_Composite.h
+++ b/Engine/Source/Editor/BehaviorTreeEditor/Classes/BehaviorTreeGraphNode_Composite.h
@@ -8,7 +8,7 @@
 #include "BehaviorTreeGraphNode_Composite.generated.h"
 
 UCLASS()
-class UBehaviorTreeGraphNode_Composite : public UBehaviorTreeGraphNode
+class BEHAVIORTREEEDITOR_API UBehaviorTreeGraphNode_Composite : public UBehaviorTreeGraphNode
 {
 	GENERATED_UCLASS_BODY()
 	
