[BT Editor]
Refresh Blackboard views after creating a new BTGraph since it may auto-assign a default Blackboard to the BT on creation (after the Blackboard views were created).

diff --git a/Engine/Source/Editor/BehaviorTreeEditor/Private/BehaviorTreeEditor.cpp b/Engine/Source/Editor/BehaviorTreeEditor/Private/BehaviorTreeEditor.cpp
--- a/Engine/Source/Editor/BehaviorTreeEditor/Private/BehaviorTreeEditor.cpp
+++ b/Engine/Source/Editor/BehaviorTreeEditor/Private/BehaviorTreeEditor.cpp
@@ -96,19 +96,24 @@ FBehaviorTreeEditor::~FBehaviorTreeEditor()
 	Debugger.Reset();
 }
 
+void FBehaviorTreeEditor::RefreshBlackboardViewsAssociatedObject()
+{
+	if (BlackboardView.IsValid())
+	{
+		BlackboardView->SetObject(GetBlackboardData());
+	}
+
+	if (BlackboardEditor.IsValid())
+	{
+		BlackboardEditor->SetObject(GetBlackboardData());
+	}
+}
+
 void FBehaviorTreeEditor::PostUndo(bool bSuccess)
 {
 	if (bSuccess)
 	{
-		if(BlackboardView.IsValid())
-		{
-			BlackboardView->SetObject(GetBlackboardData());
-		}
-
-		if(BlackboardEditor.IsValid())
-		{
-			BlackboardEditor->SetObject(GetBlackboardData());
-		}
+		RefreshBlackboardViewsAssociatedObject();
 	}
 
 	FAIGraphEditor::PostUndo(bSuccess);
@@ -118,15 +123,7 @@ void FBehaviorTreeEditor::PostRedo(bool bSuccess)
 {
 	if (bSuccess)
 	{
-		if (BlackboardView.IsValid())
-		{
-			BlackboardView->SetObject(GetBlackboardData());
-		}
-
-		if (BlackboardEditor.IsValid())
-		{
-			BlackboardEditor->SetObject(GetBlackboardData());
-		}
+		RefreshBlackboardViewsAssociatedObject();
 	}
 
 	FAIGraphEditor::PostRedo(bSuccess);
@@ -141,14 +138,7 @@ void FBehaviorTreeEditor::NotifyPostChange( const FPropertyChangedEvent& Propert
 			BlackboardData = BehaviorTree->BlackboardAsset;
 		}
 
-		if(BlackboardView.IsValid())
-		{
-			BlackboardView->SetObject(GetBlackboardData());
-		}
-		if(BlackboardEditor.IsValid())
-		{
-			BlackboardEditor->SetObject(GetBlackboardData());
-		}
+		RefreshBlackboardViewsAssociatedObject();
 	}
 }
 
@@ -316,6 +306,7 @@ void FBehaviorTreeEditor::RestoreBehaviorTree()
 	if(bNewGraph)
 	{
 		MyGraph->UpdateAsset(UBehaviorTreeGraph::ClearDebuggerFlags | UBehaviorTreeGraph::KeepRebuildCounter);
+		RefreshBlackboardViewsAssociatedObject();
 	}
 	else
 	{
@@ -384,14 +375,7 @@ float FBehaviorTreeEditor::HandleGetDebugTimeStamp(bool bUseCurrentState) const
 
 void FBehaviorTreeEditor::HandleDebuggedBlackboardChanged(UBlackboardData* InBlackboardData)
 {
-	if(BlackboardView.IsValid())
-	{
-		BlackboardView->SetObject(InBlackboardData);
-	}
-	if(BlackboardEditor.IsValid())
-	{
-		BlackboardEditor->SetObject(InBlackboardData);
-	}
+	RefreshBlackboardViewsAssociatedObject();
 }
 
 bool FBehaviorTreeEditor::HandleGetDisplayCurrentState() const
diff --git a/Engine/Source/Editor/BehaviorTreeEditor/Private/BehaviorTreeEditor.h b/Engine/Source/Editor/BehaviorTreeEditor/Private/BehaviorTreeEditor.h
--- a/Engine/Source/Editor/BehaviorTreeEditor/Private/BehaviorTreeEditor.h
+++ b/Engine/Source/Editor/BehaviorTreeEditor/Private/BehaviorTreeEditor.h
@@ -283,6 +283,9 @@ private:
 	/** Refresh the debugger's display */
 	void RefreshDebugger();
 
+	/** Push new associated Blackboard data to Blackboard views */
+	void RefreshBlackboardViewsAssociatedObject();
+
 	TSharedPtr<FDocumentTracker> DocumentManager;
 	TWeakPtr<FDocumentTabFactory> GraphEditorTabFactoryPtr;
 
