[BT Editor]
When reusing an editor it must be aware of the new objects in edition. For example, BehaviorTree may want to share an editor already opened for its associated Blackboard asset.

diff --git a/Engine/Source/Editor/BehaviorTreeEditor/Private/BehaviorTreeEditor.cpp b/Engine/Source/Editor/BehaviorTreeEditor/Private/BehaviorTreeEditor.cpp
--- a/Engine/Source/Editor/BehaviorTreeEditor/Private/BehaviorTreeEditor.cpp
+++ b/Engine/Source/Editor/BehaviorTreeEditor/Private/BehaviorTreeEditor.cpp
@@ -205,7 +205,8 @@ void FBehaviorTreeEditor::InitBehaviorTreeEditor( const EToolkitMode::Type Mode,
 		ToolbarBuilder = MakeShareable(new FBehaviorTreeEditorToolbar(SharedThis(this)));
 	}
 
-	// if we are already editing objects, dont try to recreate the editor from scratch
+	// if we are already editing objects, dont try to recreate the editor from scratch but update the list of objects in edition
+    // ex: BehaviorTree may want to reuse an editor already opened for its associated Blackboard asset.
 	const TArray<UObject*>* EditedObjects = GetObjectsCurrentlyBeingEdited();
 	if(EditedObjects == nullptr || EditedObjects->Num() == 0)
 	{
@@ -255,6 +256,14 @@ void FBehaviorTreeEditor::InitBehaviorTreeEditor( const EToolkitMode::Type Mode,
 	{
 		check(Debugger.IsValid());
 		Debugger->Setup(BehaviorTree, SharedThis(this));
+
+		for (UObject* ObjectToEdit : ObjectsToEdit)
+		{
+			if (!EditedObjects->Contains(ObjectToEdit))
+		    {
+				AddEditingObject(ObjectToEdit);
+		    }
+		}
 	}
 
 	if(BehaviorTreeToEdit != nullptr)
