[BT Editor]
Set Focus on BTEditor window when hitting breakpoint

diff --git a/Engine/Source/Editor/BehaviorTreeEditor/Private/BehaviorTreeDebugger.cpp b/Engine/Source/Editor/BehaviorTreeEditor/Private/BehaviorTreeDebugger.cpp
--- a/Engine/Source/Editor/BehaviorTreeEditor/Private/BehaviorTreeDebugger.cpp
+++ b/Engine/Source/Editor/BehaviorTreeEditor/Private/BehaviorTreeDebugger.cpp
@@ -1018,6 +1018,11 @@ void FBehaviorTreeDebugger::OnActiveNodeChanged(const TArray<uint16>& ActivePath
 
 	if (bShouldPause)
 	{
+		if (EditorOwner.IsValid())
+		{
+			EditorOwner.Pin()->FocusWindow(TreeAsset);
+		}
+
 		PausePlaySession();
 	}
 }
