BehaviorTreeDebugger will also leave kismet debugging mode when resuming session.
This is a very basic fix involving knowledge of the kismet debugger. Ideally we would need a unified flow to leave debugging mode from the different debuggers to prevent strong coupling between modules. Each debugger (Blueprint & BehaviorTree for now) could then take the appropriate actions to resume the session.

diff --git a/Engine/Source/Editor/BehaviorTreeEditor/Private/BehaviorTreeDebugger.cpp b/Engine/Source/Editor/BehaviorTreeEditor/Private/BehaviorTreeDebugger.cpp
--- a/Engine/Source/Editor/BehaviorTreeEditor/Private/BehaviorTreeDebugger.cpp
+++ b/Engine/Source/Editor/BehaviorTreeEditor/Private/BehaviorTreeDebugger.cpp
@@ -21,6 +21,7 @@
 #include "BehaviorTreeGraphNode_Service.h"
 #include "BehaviorTree/BehaviorTree.h"
 #include "BehaviorTreeDelegates.h"
+#include "Framework/Application/SlateApplication.h"
 
 FBehaviorTreeDebugger::FBehaviorTreeDebugger()
 {
@@ -1032,6 +1033,13 @@ void FBehaviorTreeDebugger::StopPlaySession()
 	if (GUnrealEd->PlayWorld)
 	{
 		GEditor->RequestEndPlayMap();
+ 
+		// @TODO: we need a unified flow to leave debugging mode from the different debuggers to prevent strong coupling between modules.
+		// Each debugger (Blueprint & BehaviorTree for now) could then take the appropriate actions to resume the session.
+  		if (FSlateApplication::Get().InKismetDebuggingMode())
+  		{
+  			FSlateApplication::Get().LeaveDebuggingMode();
+  		}
 	}
 }
 
@@ -1065,6 +1073,13 @@ void FBehaviorTreeDebugger::ResumePlaySession()
 	});
 	if(bResumed)
 	{
+		// @TODO: we need a unified flow to leave debugging mode from the different debuggers to prevent strong coupling between modules.
+		// Each debugger (Blueprint & BehaviorTree for now) could then take the appropriate actions to resume the session.
+		if (FSlateApplication::Get().InKismetDebuggingMode())
+		{
+			FSlateApplication::Get().LeaveDebuggingMode();
+		}
+
 		GUnrealEd->PlaySessionResumed();
 	}
 }
