BehaviorTreeGraph: CollectAllNodeInstances now collects all BTNodes from the Decorator subgraph (ex: UBehaviorTreeGraphNode_CompositeDecorator) instead of relying on CollectDecoratorData. The problem with CollectDecoratorData is that it returns only connected BTNodes which is not approriate when collecting all nodes to remove orphan BTNodes.
This is a modified fix from 2037569: fixed inner nodes of behavior tree's composite decorator being removed on save

diff --git a/Engine/Source/Editor/BehaviorTreeEditor/Private/BehaviorTreeGraph.cpp b/Engine/Source/Editor/BehaviorTreeEditor/Private/BehaviorTreeGraph.cpp
--- a/Engine/Source/Editor/BehaviorTreeEditor/Private/BehaviorTreeGraph.cpp
+++ b/Engine/Source/Editor/BehaviorTreeEditor/Private/BehaviorTreeGraph.cpp
@@ -24,6 +24,7 @@
 #include "BehaviorTree/Tasks/BTTask_Wait.h"
 #include "BehaviorTreeGraphNode_SimpleParallel.h"
 #include "BehaviorTreeGraphNode_SubtreeTask.h"
+#include "BehaviorTreeDecoratorGraphNode_Decorator.h"
 
 //////////////////////////////////////////////////////////////////////////
 // BehaviorTreeGraph
@@ -942,27 +943,27 @@ void UBehaviorTreeGraph::CreateBTFromGraph(UBehaviorTreeGraphNode* RootEdNode)
 	RemoveOrphanedNodes();
 }
 
-void UBehaviorTreeGraph::CollectAllNodeInstances(TSet<UObject*>& NodeInstance)
+void UBehaviorTreeGraph::CollectAllNodeInstances(TSet<UObject*>& NodeInstances)
 {
-	Super::CollectAllNodeInstances(NodeInstance);
+	Super::CollectAllNodeInstances(NodeInstances);
 
-	for (int32 Idx = 0; Idx < Nodes.Num(); Idx++)
+	for (UEdGraphNode* EdGraphNode : Nodes)
 	{
-		UBehaviorTreeGraphNode* MyNode = Cast<UBehaviorTreeGraphNode>(Nodes[Idx]);
-		if (MyNode)
+		UBehaviorTreeGraphNode* BTGraphNode = Cast<UBehaviorTreeGraphNode>(EdGraphNode);
+		if (BTGraphNode != nullptr)
 		{
-			for (int32 SubIdx = 0; SubIdx < MyNode->Decorators.Num(); SubIdx++)
+			for (UBehaviorTreeGraphNode* BTDecoratorGraphNode : BTGraphNode->Decorators)
 			{
-				UBehaviorTreeGraphNode_CompositeDecorator* SubgraphNode = Cast<UBehaviorTreeGraphNode_CompositeDecorator>(MyNode->Decorators[SubIdx]);
-				if (SubgraphNode)
+				UEdGraph* EdSubGraph = BTDecoratorGraphNode->GetBoundGraph();
+				if (EdSubGraph != nullptr)
 				{
-					TArray<UBTDecorator*> DecoratorInstances;
-					TArray<FBTDecoratorLogic> DummyOps;
-					SubgraphNode->CollectDecoratorData(DecoratorInstances, DummyOps);
-
-					for (int32 DecoratorIdx = 0; DecoratorIdx < DecoratorInstances.Num(); DecoratorIdx++)
+					for (UEdGraphNode* SubGraphNode : EdSubGraph->Nodes)
 					{
-						NodeInstance.Add(DecoratorInstances[DecoratorIdx]);
+						const UBehaviorTreeDecoratorGraphNode_Decorator* DecoratorNode = Cast<const UBehaviorTreeDecoratorGraphNode_Decorator>(SubGraphNode);
+						if (DecoratorNode && DecoratorNode->NodeInstance)
+						{
+							NodeInstances.Add(DecoratorNode->NodeInstance);
+						}
 					}
 				}
 			}
