Centralized comment support for AIGraph sublcasses, which adds support for comments to EQS graphs
#okgithub public

diff --git a/Engine/Source/Editor/AIGraph/Classes/AIGraphSchema.h b/Engine/Source/Editor/AIGraph/Classes/AIGraphSchema.h
--- a/Engine/Source/Editor/AIGraph/Classes/AIGraphSchema.h
+++ b/Engine/Source/Editor/AIGraph/Classes/AIGraphSchema.h
@@ -11,6 +11,23 @@
 class FSlateRect;
 class UEdGraph;
 
+/** Action to add a comment to the graph */
+USTRUCT()
+struct FAISchemaAction_AddComment : public FEdGraphSchemaAction
+{
+	GENERATED_BODY()
+	
+	FAISchemaAction_AddComment() : FEdGraphSchemaAction() {}
+	FAISchemaAction_AddComment(FText InDescription, FText InToolTip)
+		: FEdGraphSchemaAction(FText(), MoveTemp(InDescription), MoveTemp(InToolTip), 0)
+	{
+	}
+
+	// FEdGraphSchemaAction interface
+	virtual UEdGraphNode* PerformAction(class UEdGraph* ParentGraph, UEdGraphPin* FromPin, const FVector2D Location, bool bSelectNewNode = true) override final;
+	// End of FEdGraphSchemaAction interface
+};
+
 /** Action to add a node to the graph */
 USTRUCT()
 struct AIGRAPH_API FAISchemaAction_NewNode : public FEdGraphSchemaAction
@@ -95,6 +112,8 @@ class AIGRAPH_API UAIGraphSchema : public UEdGraphSchema
 	virtual void BreakNodeLinks(UEdGraphNode& TargetNode) const override;
 	virtual void BreakPinLinks(UEdGraphPin& TargetPin, bool bSendsNodeNotification) const override;
 	virtual void BreakSinglePinLink(UEdGraphPin* SourcePin, UEdGraphPin* TargetPin) const override;
+	virtual TSharedPtr<FEdGraphSchemaAction> GetCreateCommentAction() const override;
+	virtual int32 GetNodeSelectionCount(const UEdGraph* Graph) const override;
 	//~ End EdGraphSchema Interface
 
 	virtual void GetGraphNodeContextActions(FGraphContextMenuBuilder& ContextMenuBuilder, int32 SubNodeFlags) const;
diff --git a/Engine/Source/Editor/AIGraph/Private/AIGraphEditor.cpp b/Engine/Source/Editor/AIGraph/Private/AIGraphEditor.cpp
--- a/Engine/Source/Editor/AIGraph/Private/AIGraphEditor.cpp
+++ b/Engine/Source/Editor/AIGraph/Private/AIGraphEditor.cpp
@@ -10,6 +10,7 @@
 #include "ScopedTransaction.h"
 #include "EdGraphUtilities.h"
 #include "Framework/Commands/GenericCommands.h"
+#include "GraphEditorActions.h"
 #include "HAL/PlatformApplicationMisc.h"
 
 #define LOCTEXT_NAMESPACE "AIGraph"
@@ -78,6 +79,12 @@ void FAIGraphEditor::CreateCommandList()
 		FExecuteAction::CreateRaw(this, &FAIGraphEditor::DuplicateNodes),
 		FCanExecuteAction::CreateRaw(this, &FAIGraphEditor::CanDuplicateNodes)
 		);
+
+	GraphEditorCommands->MapAction(
+		FGraphEditorCommands::Get().CreateComment,
+		FExecuteAction::CreateRaw(this, &FAIGraphEditor::OnCreateComment),
+		FCanExecuteAction::CreateRaw(this, &FAIGraphEditor::CanCreateComment)
+	);
 }
 
 FGraphPanelSelectionSet FAIGraphEditor::GetSelectedNodes() const
@@ -483,6 +490,25 @@ bool FAIGraphEditor::CanDuplicateNodes() const
 	return CanCopyNodes();
 }
 
+bool FAIGraphEditor::CanCreateComment() const
+{
+ 	TSharedPtr<SGraphEditor> CurrentGraphEditor = UpdateGraphEdPtr.Pin();
+ 	return CurrentGraphEditor.IsValid() ? (CurrentGraphEditor->GetNumberOfSelectedNodes() != 0) : false;
+}
+
+void FAIGraphEditor::OnCreateComment()
+{
+	TSharedPtr<SGraphEditor> CurrentGraphEditor = UpdateGraphEdPtr.Pin();
+	if (UEdGraph* EdGraph = CurrentGraphEditor.IsValid() ? CurrentGraphEditor->GetCurrentGraph() : nullptr)
+	{
+		TSharedPtr<FEdGraphSchemaAction> Action = EdGraph->GetSchema()->GetCreateCommentAction();
+		if (Action.IsValid())
+		{
+			Action->PerformAction(EdGraph, nullptr, FVector2D());
+		}
+	}
+}
+
 void FAIGraphEditor::OnClassListUpdated()
 {
 	TSharedPtr<SGraphEditor> CurrentGraphEditor = UpdateGraphEdPtr.Pin();
@@ -505,4 +531,5 @@ void FAIGraphEditor::OnClassListUpdated()
 	}
 }
 
+
 #undef LOCTEXT_NAMESPACE
diff --git a/Engine/Source/Editor/AIGraph/Private/AIGraphSchema.cpp b/Engine/Source/Editor/AIGraph/Private/AIGraphSchema.cpp
--- a/Engine/Source/Editor/AIGraph/Private/AIGraphSchema.cpp
+++ b/Engine/Source/Editor/AIGraph/Private/AIGraphSchema.cpp
@@ -10,6 +10,7 @@
 #include "AIGraphConnectionDrawingPolicy.h"
 #include "ScopedTransaction.h"
 #include "Framework/Commands/GenericCommands.h"
+#include "EdGraphNode_Comment.h"
 
 #define LOCTEXT_NAMESPACE "AIGraph"
 #define SNAP_GRID (16) // @todo ensure this is the same as SNodePanel::GetSnapGridSize()
@@ -20,6 +21,30 @@ namespace
 	const int32 NodeDistance = 60;
 }
 
+//////////////////////////////////////////////////////////////////////////
+
+UEdGraphNode* FAISchemaAction_AddComment::PerformAction(class UEdGraph* ParentGraph, UEdGraphPin* FromPin, const FVector2D Location, bool bSelectNewNode)
+{
+	UEdGraphNode_Comment* const CommentTemplate = NewObject<UEdGraphNode_Comment>();
+
+	FVector2D SpawnLocation = Location;
+	FSlateRect Bounds;
+
+	TSharedPtr<SGraphEditor> GraphEditorPtr = SGraphEditor::FindGraphEditorForGraph(ParentGraph);
+	if (GraphEditorPtr.IsValid() && GraphEditorPtr->GetBoundsForSelectedNodes(/*out*/ Bounds, 50.0f))
+	{
+		CommentTemplate->SetBounds(Bounds);
+		SpawnLocation.X = CommentTemplate->NodePosX;
+		SpawnLocation.Y = CommentTemplate->NodePosY;
+	}
+
+	UEdGraphNode* const NewNode = FEdGraphSchemaAction_NewNode::SpawnNodeFromTemplate<UEdGraphNode_Comment>(ParentGraph, CommentTemplate, SpawnLocation, bSelectNewNode);
+
+	return NewNode;
+}
+
+//////////////////////////////////////////////////////////////////////////
+
 UEdGraphNode* FAISchemaAction_NewNode::PerformAction(class UEdGraph* ParentGraph, UEdGraphPin* FromPin, const FVector2D Location, bool bSelectNewNode)
 {
 	UEdGraphNode* ResultNode = NULL;
@@ -121,6 +146,7 @@ void FAISchemaAction_NewSubNode::AddReferencedObjects(FReferenceCollector& Colle
 	Collector.AddReferencedObject(NodeTemplate);
 	Collector.AddReferencedObject(ParentNode);
 }
+
 //////////////////////////////////////////////////////////////////////////
 
 UAIGraphSchema::UAIGraphSchema(const FObjectInitializer& ObjectInitializer) : Super(ObjectInitializer)
@@ -295,4 +321,23 @@ class FConnectionDrawingPolicy* UAIGraphSchema::CreateConnectionDrawingPolicy(in
 	return new FAIGraphConnectionDrawingPolicy(InBackLayerID, InFrontLayerID, InZoomFactor, InClippingRect, InDrawElements, InGraphObj);
 }
 
+TSharedPtr<FEdGraphSchemaAction> UAIGraphSchema::GetCreateCommentAction() const
+{
+	return TSharedPtr<FEdGraphSchemaAction>(static_cast<FEdGraphSchemaAction*>(new FAISchemaAction_AddComment));
+}
+
+int32 UAIGraphSchema::GetNodeSelectionCount(const UEdGraph* Graph) const
+{
+	if (Graph)
+	{
+		TSharedPtr<SGraphEditor> GraphEditorPtr = SGraphEditor::FindGraphEditorForGraph(Graph);
+		if (GraphEditorPtr.IsValid())
+		{
+			return GraphEditorPtr->GetNumberOfSelectedNodes();
+		}
+	}
+
+	return 0;
+}
+
 #undef LOCTEXT_NAMESPACE
diff --git a/Engine/Source/Editor/AIGraph/Public/AIGraphEditor.h b/Engine/Source/Editor/AIGraph/Public/AIGraphEditor.h
--- a/Engine/Source/Editor/AIGraph/Public/AIGraphEditor.h
+++ b/Engine/Source/Editor/AIGraph/Public/AIGraphEditor.h
@@ -38,6 +38,10 @@ public:
 	void DuplicateNodes();
 	bool CanDuplicateNodes() const;
 
+	bool CanCreateComment() const;
+	void OnCreateComment();
+
+
 	virtual void OnClassListUpdated();
 
 protected:
diff --git a/Engine/Source/Editor/BehaviorTreeEditor/Classes/EdGraphSchema_BehaviorTree.h b/Engine/Source/Editor/BehaviorTreeEditor/Classes/EdGraphSchema_BehaviorTree.h
--- a/Engine/Source/Editor/BehaviorTreeEditor/Classes/EdGraphSchema_BehaviorTree.h
+++ b/Engine/Source/Editor/BehaviorTreeEditor/Classes/EdGraphSchema_BehaviorTree.h
@@ -30,22 +30,6 @@ struct FBehaviorTreeSchemaAction_AutoArrange : public FEdGraphSchemaAction
 	//~ End FEdGraphSchemaAction Interface
 };
 
-USTRUCT()
-struct FBehaviorTreeSchemaAction_AddComment : public FEdGraphSchemaAction
-{
-	GENERATED_BODY()
-	
-	FBehaviorTreeSchemaAction_AddComment() : FEdGraphSchemaAction() {}
-	FBehaviorTreeSchemaAction_AddComment(FText InDescription, FText InToolTip)
-		: FEdGraphSchemaAction(FText(), MoveTemp(InDescription), MoveTemp(InToolTip), 0)
-	{
-	}
-
-	// FEdGraphSchemaAction interface
-	virtual UEdGraphNode* PerformAction(class UEdGraph* ParentGraph, UEdGraphPin* FromPin, const FVector2D Location, bool bSelectNewNode = true) override final;
-	// End of FEdGraphSchemaAction interface
-};
-
 UCLASS(MinimalAPI)
 class UEdGraphSchema_BehaviorTree : public UAIGraphSchema
 {
@@ -59,11 +43,9 @@ class UEdGraphSchema_BehaviorTree : public UAIGraphSchema
 	virtual const FPinConnectionResponse CanMergeNodes(const UEdGraphNode* A, const UEdGraphNode* B) const override;
 	virtual FLinearColor GetPinTypeColor(const FEdGraphPinType& PinType) const override;
 	virtual class FConnectionDrawingPolicy* CreateConnectionDrawingPolicy(int32 InBackLayerID, int32 InFrontLayerID, float InZoomFactor, const FSlateRect& InClippingRect, class FSlateWindowElementList& InDrawElements, class UEdGraph* InGraphObj) const override;
-	virtual int32 GetNodeSelectionCount(const UEdGraph* Graph) const override;
 	virtual bool IsCacheVisualizationOutOfDate(int32 InVisualizationCacheID) const override;
 	virtual int32 GetCurrentVisualizationCacheID() const override;
 	virtual void ForceVisualizationCacheClear() const override;
-	virtual TSharedPtr<FEdGraphSchemaAction> GetCreateCommentAction() const override;
 	//~ End EdGraphSchema Interface
 
 	virtual void GetGraphNodeContextActions(FGraphContextMenuBuilder& ContextMenuBuilder, int32 SubNodeFlags) const override;
diff --git a/Engine/Source/Editor/BehaviorTreeEditor/Private/BehaviorTreeEditor.cpp b/Engine/Source/Editor/BehaviorTreeEditor/Private/BehaviorTreeEditor.cpp
--- a/Engine/Source/Editor/BehaviorTreeEditor/Private/BehaviorTreeEditor.cpp
+++ b/Engine/Source/Editor/BehaviorTreeEditor/Private/BehaviorTreeEditor.cpp
@@ -77,7 +77,6 @@ FBehaviorTreeEditor::FBehaviorTreeEditor()
 	bShowDecoratorRangeSelf = false;
 	bSelectedNodeIsInjected = false;
 	bSelectedNodeIsRootLevel = false;
-	SelectedNodesCount = 0;
 
 	bHasMultipleTaskBP = false;
 	bHasMultipleDecoratorBP = false;
@@ -560,12 +559,6 @@ TSharedRef<SGraphEditor> FBehaviorTreeEditor::CreateGraphEditorWidget(UEdGraph*
 			FIsActionChecked(),
 			FIsActionButtonVisible::CreateSP( this, &FBehaviorTreeEditor::CanToggleBreakpoint )
 			);
-
-		GraphEditorCommands->MapAction(
-			FGraphEditorCommands::Get().CreateComment,
-			FExecuteAction::CreateSP(this, &FBehaviorTreeEditor::OnCreateComment),
-			FCanExecuteAction::CreateSP(this, &FBehaviorTreeEditor::CanCreateComment)
-			);
 	}
 
 	SGraphEditor::FGraphEditorEvents InEvents;
@@ -731,8 +724,6 @@ TSharedRef<SWidget> FBehaviorTreeEditor::SpawnBlackboardDetails()
 
 void FBehaviorTreeEditor::OnSelectedNodesChanged(const TSet<class UObject*>& NewSelection)
 {
-	SelectedNodesCount = NewSelection.Num();
-
 	BehaviorTreeEditorUtils::FPropertySelectionInfo SelectionInfo;
 	TArray<UObject*> Selection = BehaviorTreeEditorUtils::GetSelectionForPropertyEditor(NewSelection, SelectionInfo);
 
@@ -1598,12 +1589,6 @@ void FBehaviorTreeEditor::OnNodeTitleCommitted(const FText& NewText, ETextCommit
 	}
 }
 
-bool FBehaviorTreeEditor::GetBoundsForSelectedNodes(FSlateRect& Rect, float Padding) const
-{
-	TSharedPtr<SGraphEditor> FocusedGraphEd = UpdateGraphEdPtr.Pin();
-	return FocusedGraphEd.IsValid() && FocusedGraphEd->GetBoundsForSelectedNodes(Rect, Padding);
-}
-
 void FBehaviorTreeEditor::InitializeDebuggerState(class FBehaviorTreeDebugger* ParentDebugger) const
 {
 	if (Debugger.IsValid())
@@ -1839,21 +1824,4 @@ bool FBehaviorTreeEditor::CanCreateNewBlackboard() const
 	return !IsDebuggerReady();
 }
 
-bool FBehaviorTreeEditor::CanCreateComment() const
-{
-	return (SelectedNodesCount > 0);
-}
-
-void FBehaviorTreeEditor::OnCreateComment()
-{
-	if (BehaviorTree && BehaviorTree->BTGraph)
-	{
-		TSharedPtr<FEdGraphSchemaAction> Action = BehaviorTree->BTGraph->GetSchema()->GetCreateCommentAction();
-		if (Action.IsValid())
-		{
-			Action->PerformAction(BehaviorTree->BTGraph, nullptr, FVector2D());
-		}
-	}
-}
-
 #undef LOCTEXT_NAMESPACE
diff --git a/Engine/Source/Editor/BehaviorTreeEditor/Private/BehaviorTreeEditor.h b/Engine/Source/Editor/BehaviorTreeEditor/Private/BehaviorTreeEditor.h
--- a/Engine/Source/Editor/BehaviorTreeEditor/Private/BehaviorTreeEditor.h
+++ b/Engine/Source/Editor/BehaviorTreeEditor/Private/BehaviorTreeEditor.h
@@ -48,12 +48,10 @@ public:
 	//~ End IToolkit Interface
 
 	//~ Begin IBehaviorTreeEditor Interface
-	virtual uint32 GetSelectedNodesCount() const override { return SelectedNodesCount; }
 	virtual void InitializeDebuggerState(class FBehaviorTreeDebugger* ParentDebugger) const override;
 	virtual UEdGraphNode* FindInjectedNode(int32 Index) const override;
 	virtual void DoubleClickNode(class UEdGraphNode* Node) override;
 	virtual void FocusWindow(UObject* ObjectToFocusOn = NULL) override;
-	virtual bool GetBoundsForSelectedNodes(class FSlateRect& Rect, float Padding) const override;
 	//~ End IBehaviorTreeEditor Interface
 
 	//~ Begin FEditorUndoClient Interface
@@ -85,8 +83,6 @@ public:
 	bool CanAddBreakpoint() const;
 	void OnRemoveBreakpoint();
 	bool CanRemoveBreakpoint() const;
-	void OnCreateComment();
-	bool CanCreateComment() const;
 
 	void SearchTree();
 	bool CanSearchTree() const;
@@ -310,7 +306,6 @@ private:
 	uint32 bForceDisablePropertyEdit : 1;
 	uint32 bSelectedNodeIsInjected : 1;
 	uint32 bSelectedNodeIsRootLevel : 1;
-	uint32 SelectedNodesCount;
 
 	uint32 bHasMultipleTaskBP : 1;
 	uint32 bHasMultipleDecoratorBP : 1;
diff --git a/Engine/Source/Editor/BehaviorTreeEditor/Private/EdGraphSchema_BehaviorTree.cpp b/Engine/Source/Editor/BehaviorTreeEditor/Private/EdGraphSchema_BehaviorTree.cpp
--- a/Engine/Source/Editor/BehaviorTreeEditor/Private/EdGraphSchema_BehaviorTree.cpp
+++ b/Engine/Source/Editor/BehaviorTreeEditor/Private/EdGraphSchema_BehaviorTree.cpp
@@ -49,38 +49,6 @@ UEdGraphNode* FBehaviorTreeSchemaAction_AutoArrange::PerformAction(class UEdGrap
 	return NULL;
 }
 
-//----------------------------------------------------------------------//
-// 
-//----------------------------------------------------------------------//
-UEdGraphNode* FBehaviorTreeSchemaAction_AddComment::PerformAction(class UEdGraph* ParentGraph, UEdGraphPin* FromPin, const FVector2D Location, bool bSelectNewNode)
-{
-	UEdGraphNode_Comment* const CommentTemplate = NewObject<UEdGraphNode_Comment>();
-
-	FVector2D SpawnLocation = Location;
-
-	TSharedPtr<IBehaviorTreeEditor> BTEditor;
-	if (UBehaviorTree* const BTAsset = Cast<UBehaviorTree>(ParentGraph->GetOuter()))
-	{
-		TSharedPtr<IToolkit> BTAssetEditor = FToolkitManager::Get().FindEditorForAsset(BTAsset);
-		if (BTAssetEditor.IsValid())
-		{
-			BTEditor = StaticCastSharedPtr<IBehaviorTreeEditor>(BTAssetEditor);
-		}
-	}
-
-	FSlateRect Bounds;
-	if (BTEditor.IsValid() && BTEditor->GetBoundsForSelectedNodes(Bounds, 50.0f))
-	{
-		CommentTemplate->SetBounds(Bounds);
-		SpawnLocation.X = CommentTemplate->NodePosX;
-		SpawnLocation.Y = CommentTemplate->NodePosY;
-	}
-
-	UEdGraphNode* const NewNode = FEdGraphSchemaAction_NewNode::SpawnNodeFromTemplate<UEdGraphNode_Comment>(ParentGraph, CommentTemplate, SpawnLocation, bSelectNewNode);
-
-	return NewNode;
-}
-
 //----------------------------------------------------------------------//
 // 
 //----------------------------------------------------------------------//
@@ -450,28 +418,6 @@ class FConnectionDrawingPolicy* UEdGraphSchema_BehaviorTree::CreateConnectionDra
 	return new FBehaviorTreeConnectionDrawingPolicy(InBackLayerID, InFrontLayerID, InZoomFactor, InClippingRect, InDrawElements, InGraphObj);
 }
 
-int32 UEdGraphSchema_BehaviorTree::GetNodeSelectionCount(const UEdGraph* Graph) const
-{
-	if (Graph)
-	{
-		TSharedPtr<IBehaviorTreeEditor> BTEditor;
-		if (UBehaviorTree* BTAsset = Cast<UBehaviorTree>(Graph->GetOuter()))
-		{
-			TSharedPtr< IToolkit > BTAssetEditor = FToolkitManager::Get().FindEditorForAsset(BTAsset);
-			if (BTAssetEditor.IsValid())
-			{
-				BTEditor = StaticCastSharedPtr<IBehaviorTreeEditor>(BTAssetEditor);
-			}
-		}
-		if(BTEditor.IsValid())
-		{
-			return BTEditor->GetSelectedNodesCount();
-		}
-	}
-
-	return 0;
-}
-
 bool UEdGraphSchema_BehaviorTree::IsCacheVisualizationOutOfDate(int32 InVisualizationCacheID) const
 {
 	return CurrentCacheRefreshID != InVisualizationCacheID;
@@ -487,9 +433,4 @@ void UEdGraphSchema_BehaviorTree::ForceVisualizationCacheClear() const
 	++CurrentCacheRefreshID;
 }
 
-TSharedPtr<FEdGraphSchemaAction> UEdGraphSchema_BehaviorTree::GetCreateCommentAction() const
-{
-	return TSharedPtr<FEdGraphSchemaAction>(static_cast<FEdGraphSchemaAction*>(new FBehaviorTreeSchemaAction_AddComment));
-}
-
 #undef LOCTEXT_NAMESPACE
diff --git a/Engine/Source/Editor/BehaviorTreeEditor/Public/IBehaviorTreeEditor.h b/Engine/Source/Editor/BehaviorTreeEditor/Public/IBehaviorTreeEditor.h
--- a/Engine/Source/Editor/BehaviorTreeEditor/Public/IBehaviorTreeEditor.h
+++ b/Engine/Source/Editor/BehaviorTreeEditor/Public/IBehaviorTreeEditor.h
@@ -13,12 +13,9 @@ class IBehaviorTreeEditor : public FWorkflowCentricApplication
 {
 
 public:
-	virtual uint32 GetSelectedNodesCount() const = 0;
-	
 	virtual void InitializeDebuggerState(class FBehaviorTreeDebugger* ParentDebugger) const = 0;
 	virtual UEdGraphNode* FindInjectedNode(int32 Index) const = 0;
 	virtual void DoubleClickNode(class UEdGraphNode* Node) = 0;
-	virtual bool GetBoundsForSelectedNodes(class FSlateRect& Rect, float Padding) const = 0;
 };
 
 
diff --git a/Engine/Source/Editor/EnvironmentQueryEditor/Classes/EdGraphSchema_EnvironmentQuery.h b/Engine/Source/Editor/EnvironmentQueryEditor/Classes/EdGraphSchema_EnvironmentQuery.h
--- a/Engine/Source/Editor/EnvironmentQueryEditor/Classes/EdGraphSchema_EnvironmentQuery.h
+++ b/Engine/Source/Editor/EnvironmentQueryEditor/Classes/EdGraphSchema_EnvironmentQuery.h
@@ -21,7 +21,6 @@ class UEdGraphSchema_EnvironmentQuery : public UAIGraphSchema
 	virtual void GetGraphContextActions(FGraphContextMenuBuilder& ContextMenuBuilder) const override;
 	virtual const FPinConnectionResponse CanCreateConnection(const UEdGraphPin* A, const UEdGraphPin* B) const override;
 	virtual const FPinConnectionResponse CanMergeNodes(const UEdGraphNode* A, const UEdGraphNode* B) const override;
-	virtual int32 GetNodeSelectionCount(const UEdGraph* Graph) const override;
 	//~ End EdGraphSchema Interface
 
 	//~ Begin UAIGraphSchema Interface
diff --git a/Engine/Source/Editor/EnvironmentQueryEditor/Private/EdGraphSchema_EnvironmentQuery.cpp b/Engine/Source/Editor/EnvironmentQueryEditor/Private/EdGraphSchema_EnvironmentQuery.cpp
--- a/Engine/Source/Editor/EnvironmentQueryEditor/Private/EdGraphSchema_EnvironmentQuery.cpp
+++ b/Engine/Source/Editor/EnvironmentQueryEditor/Private/EdGraphSchema_EnvironmentQuery.cpp
@@ -131,26 +131,4 @@ const FPinConnectionResponse UEdGraphSchema_EnvironmentQuery::CanMergeNodes(cons
 	return FPinConnectionResponse(CONNECT_RESPONSE_DISALLOW, TEXT(""));
 }
 
-int32 UEdGraphSchema_EnvironmentQuery::GetNodeSelectionCount(const UEdGraph* Graph) const
-{
-	if (Graph)
-	{
-		TSharedPtr<IEnvironmentQueryEditor> EnvQueryEditor;
-		if (UEnvQuery* QueryAsset = Cast<UEnvQuery>(Graph->GetOuter()))
-		{
-			TSharedPtr< IToolkit > QueryAssetEditor = FToolkitManager::Get().FindEditorForAsset(QueryAsset);
-			if (QueryAssetEditor.IsValid())
-			{
-				EnvQueryEditor = StaticCastSharedPtr<IEnvironmentQueryEditor>(QueryAssetEditor);
-			}
-		}
-		if (EnvQueryEditor.IsValid())
-		{
-			return EnvQueryEditor->GetSelectedNodesCount();
-		}
-	}
-
-	return 0;
-}
-
 #undef LOCTEXT_NAMESPACE
diff --git a/Engine/Source/Editor/EnvironmentQueryEditor/Private/EnvironmentQueryEditor.cpp b/Engine/Source/Editor/EnvironmentQueryEditor/Private/EnvironmentQueryEditor.cpp
--- a/Engine/Source/Editor/EnvironmentQueryEditor/Private/EnvironmentQueryEditor.cpp
+++ b/Engine/Source/Editor/EnvironmentQueryEditor/Private/EnvironmentQueryEditor.cpp
@@ -92,7 +92,6 @@ void FEnvironmentQueryEditor::UnregisterTabSpawners(const TSharedRef<class FTabM
 
 void FEnvironmentQueryEditor::InitEnvironmentQueryEditor( const EToolkitMode::Type Mode, const TSharedPtr< class IToolkitHost >& InitToolkitHost, UEnvQuery* InScript )
 {
-	SelectedNodesCount = 0;
 	Query = InScript;
 	check(Query != NULL);
 
@@ -282,7 +281,6 @@ void FEnvironmentQueryEditor::OnSelectedNodesChanged(const TSet<class UObject*>&
 {
 	TArray<UObject*> Selection;
 
-	SelectedNodesCount = NewSelection.Num();
 	if (NewSelection.Num())
 	{
 		for(TSet<class UObject*>::TConstIterator SetIt(NewSelection);SetIt;++SetIt)
diff --git a/Engine/Source/Editor/EnvironmentQueryEditor/Private/EnvironmentQueryEditor.h b/Engine/Source/Editor/EnvironmentQueryEditor/Private/EnvironmentQueryEditor.h
--- a/Engine/Source/Editor/EnvironmentQueryEditor/Private/EnvironmentQueryEditor.h
+++ b/Engine/Source/Editor/EnvironmentQueryEditor/Private/EnvironmentQueryEditor.h
@@ -29,10 +29,6 @@ public:
 	virtual FLinearColor GetWorldCentricTabColorScale() const override;
 	//~ End IToolkit Interface
 
-	//~ Begin IEnvironmentQueryEditor Interface
-	virtual uint32 GetSelectedNodesCount() const override { return SelectedNodesCount; }
-	//~ End IEnvironmentQueryEditor Interface
-
 	void OnFinishedChangingProperties(const FPropertyChangedEvent& PropertyChangedEvent);
 	void OnStatsDataChange();
 
diff --git a/Engine/Source/Editor/EnvironmentQueryEditor/Public/IEnvironmentQueryEditor.h b/Engine/Source/Editor/EnvironmentQueryEditor/Public/IEnvironmentQueryEditor.h
--- a/Engine/Source/Editor/EnvironmentQueryEditor/Public/IEnvironmentQueryEditor.h
+++ b/Engine/Source/Editor/EnvironmentQueryEditor/Public/IEnvironmentQueryEditor.h
@@ -12,7 +12,6 @@ class IEnvironmentQueryEditor : public FAssetEditorToolkit
 
 public:
 
-	virtual uint32 GetSelectedNodesCount() const = 0;
 };
 
 
