Added minimal plumbing to allow UBlackboardKeySelector to be used outside BehaviorTree editor
- Changed FBlackboardSelectorDetails to find related Blackboard based on IBlackboardAssetProvider instead of UBTNode
- Added delegates to update the cached contents of FBlackboardSelectorDetails when BlackboardData or related BBD changes

diff --git a/Engine/Source/Editor/BehaviorTreeEditor/Private/DetailCustomizations/BlackboardSelectorDetails.cpp b/Engine/Source/Editor/BehaviorTreeEditor/Private/DetailCustomizations/BlackboardSelectorDetails.cpp
--- a/Engine/Source/Editor/BehaviorTreeEditor/Private/DetailCustomizations/BlackboardSelectorDetails.cpp
+++ b/Engine/Source/Editor/BehaviorTreeEditor/Private/DetailCustomizations/BlackboardSelectorDetails.cpp
@@ -8,7 +8,7 @@
 #include "Widgets/Text/STextBlock.h"
 #include "Framework/MultiBox/MultiBoxBuilder.h"
 #include "Widgets/Input/SComboButton.h"
-#include "BehaviorTree/BTNode.h"
+#include "BehaviorTree/BlackboardAssetProvider.h"
 #include "BehaviorTreeDebugger.h"
 #include "DetailWidgetRow.h"
 #include "DetailLayoutBuilder.h"
@@ -61,10 +61,10 @@ const UBlackboardData* FBlackboardSelectorDetails::FindBlackboardAsset(UObject*
 {
 	for (UObject* TestOb = InObj; TestOb; TestOb = TestOb->GetOuter())
 	{
-		UBTNode* NodeOb = Cast<UBTNode>(TestOb);
-		if (NodeOb)
+		IBlackboardAssetProvider* Provider = Cast<IBlackboardAssetProvider>(TestOb);
+		if (Provider)
 		{
-			return NodeOb->GetBlackboardAsset();
+			return Provider->GetBlackboardAsset();
 		}
 	}
 
@@ -111,6 +111,7 @@ void FBlackboardSelectorDetails::CacheBlackboardData()
 		UBlackboardData* BlackboardAsset = const_cast<UBlackboardData*>(FindBlackboardAsset(MyObjects[ObjectIdx]));
 		if (BlackboardAsset)
 		{
+			CachedBlackboardAssetOwner = MyObjects[ObjectIdx];
 			CachedBlackboardAsset = BlackboardAsset;
 
 			TArray<FName> ProcessedNames;
@@ -153,7 +154,37 @@ void FBlackboardSelectorDetails::CacheBlackboardData()
 	if (GetDefault<UEditorPerProjectUserSettings>()->bDisplayBlackboardKeysInAlphabeticalOrder)
 	{
 		KeyValues.Sort([](const FName& a, const FName& b) { return a.LexicalLess(b); });
-	}	
+	}
+
+	if (!OnBlackboardDataChangedHandle.IsValid())
+	{
+		UBlackboardData::OnBlackboardDataChanged.AddSP(this, &FBlackboardSelectorDetails::OnBlackboardDataChanged);
+	}
+	if (!OnBlackboardOwnerChangedHandle.IsValid())
+	{
+		IBlackboardAssetProvider::OnBlackboardOwnerChanged.AddSP(this, &FBlackboardSelectorDetails::OnBlackboardOwnerChanged);
+	}
+
+}
+
+void FBlackboardSelectorDetails::OnBlackboardDataChanged(UBlackboardData* Asset)
+{
+	UBlackboardData* CachedAsset = CachedBlackboardAsset.Get();
+	if (CachedAsset == nullptr || CachedAsset == Asset)
+	{
+		CacheBlackboardData();
+		InitKeyFromProperty();
+	}
+}
+
+void FBlackboardSelectorDetails::OnBlackboardOwnerChanged(UObject* Owner, UBlackboardData* Asset)
+{
+	UObject* CachedAssetOwner = CachedBlackboardAssetOwner.Get();
+	if (CachedAssetOwner == nullptr || CachedAssetOwner == Owner)
+	{
+		CacheBlackboardData();
+		InitKeyFromProperty();
+	}
 }
 
 void FBlackboardSelectorDetails::InitKeyFromProperty()
diff --git a/Engine/Source/Editor/BehaviorTreeEditor/Private/DetailCustomizations/BlackboardSelectorDetails.h b/Engine/Source/Editor/BehaviorTreeEditor/Private/DetailCustomizations/BlackboardSelectorDetails.h
--- a/Engine/Source/Editor/BehaviorTreeEditor/Private/DetailCustomizations/BlackboardSelectorDetails.h
+++ b/Engine/Source/Editor/BehaviorTreeEditor/Private/DetailCustomizations/BlackboardSelectorDetails.h
@@ -25,12 +25,18 @@ private:
 	void CacheBlackboardData();
 	const class UBlackboardData* FindBlackboardAsset(UObject* InObj);
 
+	void OnBlackboardDataChanged(UBlackboardData* Asset);
+	void OnBlackboardOwnerChanged(UObject* Owner, UBlackboardData* Asset);
+
 	void InitKeyFromProperty();
 	void OnKeyComboChange(int32 Index);
 	TSharedRef<SWidget> OnGetKeyContent() const;
 	FText GetCurrentKeyDesc() const;
 	bool IsEditingEnabled() const;
 
+	FDelegateHandle	OnBlackboardDataChangedHandle;
+	FDelegateHandle	OnBlackboardOwnerChangedHandle;
+
 	TSharedPtr<IPropertyHandle> MyStructProperty;
 	TSharedPtr<IPropertyHandle> MyKeyNameProperty;
 	TSharedPtr<IPropertyHandle> MyKeyIDProperty;
@@ -43,6 +49,7 @@ private:
 
 	/** cached blackboard asset */
 	TWeakObjectPtr<class UBlackboardData> CachedBlackboardAsset;
+	TWeakObjectPtr<class UObject> CachedBlackboardAssetOwner;
 
 	/** property utils */
 	class IPropertyUtilities* PropUtils;
diff --git a/Engine/Source/Runtime/AIModule/Classes/BehaviorTree/BTNode.h b/Engine/Source/Runtime/AIModule/Classes/BehaviorTree/BTNode.h
--- a/Engine/Source/Runtime/AIModule/Classes/BehaviorTree/BTNode.h
+++ b/Engine/Source/Runtime/AIModule/Classes/BehaviorTree/BTNode.h
@@ -7,6 +7,7 @@
 #include "UObject/Object.h"
 #include "BehaviorTree/BehaviorTreeTypes.h"
 #include "BehaviorTree/BehaviorTreeComponent.h"
+#include "BehaviorTree/BlackboardAssetProvider.h"
 #include "GameplayTaskOwnerInterface.h"
 #include "Tasks/AITask.h"
 #include "BTNode.generated.h"
@@ -33,7 +34,7 @@ struct FBTInstancedNodeMemory
 };
 
 UCLASS(Abstract,config=Game)
-class AIMODULE_API UBTNode : public UObject, public IGameplayTaskOwnerInterface
+class AIMODULE_API UBTNode : public UObject, public IGameplayTaskOwnerInterface, public IBlackboardAssetProvider
 {
 	GENERATED_UCLASS_BODY()
 
@@ -132,8 +133,10 @@ class AIMODULE_API UBTNode : public UObject, public IGameplayTaskOwnerInterface
 	/** @return tree asset */
 	UBehaviorTree* GetTreeAsset() const;
 
+	// BEGIN IBlackboardAssetProvider
 	/** @return blackboard asset */
-	UBlackboardData* GetBlackboardAsset() const;
+	virtual UBlackboardData* GetBlackboardAsset() const override;
+	// END IBlackboardAssetProvider
 
 	/** @return node instance if bCreateNodeInstance was set */
 	UBTNode* GetNodeInstance(const UBehaviorTreeComponent& OwnerComp, uint8* NodeMemory) const;
diff --git a/Engine/Source/Runtime/AIModule/Classes/BehaviorTree/BlackboardAssetProvider.h b/Engine/Source/Runtime/AIModule/Classes/BehaviorTree/BlackboardAssetProvider.h
new file mode 100644
--- /dev/null
+++ b/Engine/Source/Runtime/AIModule/Classes/BehaviorTree/BlackboardAssetProvider.h
@@ -0,0 +1,35 @@
+// Copyright 1998-2019 Epic Games, Inc. All Rights Reserved.
+
+#pragma once
+
+#include "CoreMinimal.h"
+#include "UObject/ObjectMacros.h"
+#include "UObject/Interface.h"
+#include "BlackboardAssetProvider.generated.h"
+
+
+/** Helper interface to allow FBlackboardKeySelector properties on DataAssets (and more).
+ *  Used by FBlackboardSelectorDetails to access the related Blackboard based on UObject
+ *  hierarchy. The asset containing the Blackboard should broadcast OnBlackboardOwnerChanged
+ *  when ever the asset ptr changes. */
+UINTERFACE(Blueprintable, MinimalAPI, meta = (CannotImplementInterfaceInBlueprint))
+class UBlackboardAssetProvider : public UInterface
+{
+	GENERATED_UINTERFACE_BODY()
+};
+
+class AIMODULE_API IBlackboardAssetProvider
+{
+	GENERATED_IINTERFACE_BODY()
+
+#if WITH_EDITOR
+	/** Delegate to be called by class implementing IBlackboardAssetProvider when the property containing the returned BlackboardData is changed (i.e. on PostEditChangeProperty). */
+	DECLARE_MULTICAST_DELEGATE_TwoParams(FBlackboardOwnerChanged, UObject* /*AssetOwner*/, UBlackboardData* /*Asset*/);
+	static FBlackboardOwnerChanged OnBlackboardOwnerChanged;
+#endif
+	/** Returns BlackboardData referenced by the owner object. */
+	UFUNCTION(BlueprintCallable, Category = GameplayTags)
+	virtual UBlackboardData* GetBlackboardAsset() const PURE_VIRTUAL(IBlackboardAssetProvider::GetBlackboardAsset, return nullptr; );
+};
+
+
diff --git a/Engine/Source/Runtime/AIModule/Classes/BehaviorTree/BlackboardData.h b/Engine/Source/Runtime/AIModule/Classes/BehaviorTree/BlackboardData.h
--- a/Engine/Source/Runtime/AIModule/Classes/BehaviorTree/BlackboardData.h
+++ b/Engine/Source/Runtime/AIModule/Classes/BehaviorTree/BlackboardData.h
@@ -129,6 +129,12 @@ public:
 		return CreatedKeyType;
 	}
 
+#if WITH_EDITOR
+	/** A delegate called on PostEditChangeProperty. Can be used in editor to react to asset changes. */
+	DECLARE_MULTICAST_DELEGATE_OneParam(FBlackboardDataChanged, UBlackboardData* /*Asset*/);
+	static FBlackboardDataChanged OnBlackboardDataChanged;
+#endif
+
 	/** delegate called for every loaded blackboard asset
 	 *  meant for adding game specific persistent keys */
 	static FKeyUpdate OnUpdateKeys;
diff --git a/Engine/Source/Runtime/AIModule/Private/BehaviorTree/BlackboardAssetProvider.cpp b/Engine/Source/Runtime/AIModule/Private/BehaviorTree/BlackboardAssetProvider.cpp
new file mode 100644
--- /dev/null
+++ b/Engine/Source/Runtime/AIModule/Private/BehaviorTree/BlackboardAssetProvider.cpp
@@ -0,0 +1,17 @@
+// Copyright 1998-2019 Epic Games, Inc. All Rights Reserved.
+
+#include "BehaviorTree/BlackboardAssetProvider.h"
+
+
+//----------------------------------------------------------------------//
+// UBlackboardAssetProvider
+//----------------------------------------------------------------------//
+
+#if WITH_EDITOR
+IBlackboardAssetProvider::FBlackboardOwnerChanged IBlackboardAssetProvider::OnBlackboardOwnerChanged;
+#endif
+
+UBlackboardAssetProvider::UBlackboardAssetProvider(const FObjectInitializer& ObjectInitializer)
+	: Super(ObjectInitializer)
+{
+}
diff --git a/Engine/Source/Runtime/AIModule/Private/BehaviorTree/BlackboardData.cpp b/Engine/Source/Runtime/AIModule/Private/BehaviorTree/BlackboardData.cpp
--- a/Engine/Source/Runtime/AIModule/Private/BehaviorTree/BlackboardData.cpp
+++ b/Engine/Source/Runtime/AIModule/Private/BehaviorTree/BlackboardData.cpp
@@ -8,6 +8,9 @@
 #include "UObject/UObjectIterator.h"
 
 UBlackboardData::FKeyUpdate UBlackboardData::OnUpdateKeys;
+#if WITH_EDITOR
+UBlackboardData::FBlackboardDataChanged UBlackboardData::OnBlackboardDataChanged;
+#endif
 
 static void UpdatePersistentKeys(UBlackboardData& Asset)
 {
@@ -171,6 +174,8 @@ void UBlackboardData::PostEditChangeProperty(FPropertyChangedEvent& PropertyChan
 			PropagateKeyChangesToDerivedBlackboardAssets();
 		}
 	}
+
+	UBlackboardData::OnBlackboardDataChanged.Broadcast(this);
 }
 #endif // WITH_EDITOR
 
