Fixed FBlackboardSelectorDetails to cache correct blackboard asset owner.
The bug prevented changing blackboard to update the key dropdown menus.

diff --git a/Engine/Source/Editor/BehaviorTreeEditor/Private/DetailCustomizations/BlackboardSelectorDetails.cpp b/Engine/Source/Editor/BehaviorTreeEditor/Private/DetailCustomizations/BlackboardSelectorDetails.cpp
--- a/Engine/Source/Editor/BehaviorTreeEditor/Private/DetailCustomizations/BlackboardSelectorDetails.cpp
+++ b/Engine/Source/Editor/BehaviorTreeEditor/Private/DetailCustomizations/BlackboardSelectorDetails.cpp
@@ -57,18 +57,25 @@ void FBlackboardSelectorDetails::CustomizeChildren( TSharedRef<class IPropertyHa
 {
 }
 
-const UBlackboardData* FBlackboardSelectorDetails::FindBlackboardAsset(UObject* InObj)
+void FBlackboardSelectorDetails::FindBlackboardAsset(const UObject* InObj, const UObject*& OutBlackboardOwner, UBlackboardData*& OutBlackboardAsset) const
 {
-	for (UObject* TestOb = InObj; TestOb; TestOb = TestOb->GetOuter())
+	OutBlackboardOwner = nullptr;
+	OutBlackboardAsset = nullptr;
+
+	// Find first blackboard provider and return its' data.
+	// The data might be null if no blackboard is set up yet by the provider.
+	// It is important to always return the same provider so that the invalidate check
+	// in OnBlackboardOwnerChanged works consistently.
+	for (const UObject* TestOb = InObj; TestOb; TestOb = TestOb->GetOuter())
 	{
-		IBlackboardAssetProvider* Provider = Cast<IBlackboardAssetProvider>(TestOb);
+		const IBlackboardAssetProvider* Provider = Cast<IBlackboardAssetProvider>(TestOb);
 		if (Provider)
 		{
-			return Provider->GetBlackboardAsset();
+			OutBlackboardOwner = TestOb;
+			OutBlackboardAsset = Provider->GetBlackboardAsset();
+			break;
 		}
 	}
-
-	return NULL;
 }
 
 void FBlackboardSelectorDetails::CacheBlackboardData()
@@ -108,10 +115,13 @@ void FBlackboardSelectorDetails::CacheBlackboardData()
 	MyStructProperty->GetOuterObjects(MyObjects);
 	for (int32 ObjectIdx = 0; ObjectIdx < MyObjects.Num(); ObjectIdx++)
 	{
-		UBlackboardData* BlackboardAsset = const_cast<UBlackboardData*>(FindBlackboardAsset(MyObjects[ObjectIdx]));
+		const UObject* BlackboardOwner = nullptr;
+		UBlackboardData* BlackboardAsset = nullptr;
+		FindBlackboardAsset(MyObjects[ObjectIdx], BlackboardOwner, BlackboardAsset);
+
 		if (BlackboardAsset)
 		{
-			CachedBlackboardAssetOwner = MyObjects[ObjectIdx];
+			CachedBlackboardAssetOwner = BlackboardOwner;
 			CachedBlackboardAsset = BlackboardAsset;
 
 			TArray<FName> ProcessedNames;
@@ -179,7 +189,7 @@ void FBlackboardSelectorDetails::OnBlackboardDataChanged(UBlackboardData* Asset)
 
 void FBlackboardSelectorDetails::OnBlackboardOwnerChanged(UObject* Owner, UBlackboardData* Asset)
 {
-	UObject* CachedAssetOwner = CachedBlackboardAssetOwner.Get();
+	const UObject* CachedAssetOwner = CachedBlackboardAssetOwner.Get();
 	if (CachedAssetOwner == nullptr || CachedAssetOwner == Owner)
 	{
 		CacheBlackboardData();
diff --git a/Engine/Source/Editor/BehaviorTreeEditor/Private/DetailCustomizations/BlackboardSelectorDetails.h b/Engine/Source/Editor/BehaviorTreeEditor/Private/DetailCustomizations/BlackboardSelectorDetails.h
--- a/Engine/Source/Editor/BehaviorTreeEditor/Private/DetailCustomizations/BlackboardSelectorDetails.h
+++ b/Engine/Source/Editor/BehaviorTreeEditor/Private/DetailCustomizations/BlackboardSelectorDetails.h
@@ -23,7 +23,7 @@ public:
 private:
 
 	void CacheBlackboardData();
-	const class UBlackboardData* FindBlackboardAsset(UObject* InObj);
+	void FindBlackboardAsset(const UObject* InObj, const UObject*& OutBlackboardOwner, UBlackboardData*& OutBlackboardAsset) const;
 
 	void OnBlackboardDataChanged(UBlackboardData* Asset);
 	void OnBlackboardOwnerChanged(UObject* Owner, UBlackboardData* Asset);
@@ -49,7 +49,7 @@ private:
 
 	/** cached blackboard asset */
 	TWeakObjectPtr<class UBlackboardData> CachedBlackboardAsset;
-	TWeakObjectPtr<class UObject> CachedBlackboardAssetOwner;
+	TWeakObjectPtr<const class UObject> CachedBlackboardAssetOwner;
 
 	/** property utils */
 	class IPropertyUtilities* PropUtils;
