Fix-up strict compile errors from prior PCH updates
#rnx

diff --git a/Engine/Source/Editor/BehaviorTreeEditor/Private/BehaviorTreeDebugger.cpp b/Engine/Source/Editor/BehaviorTreeEditor/Private/BehaviorTreeDebugger.cpp
--- a/Engine/Source/Editor/BehaviorTreeEditor/Private/BehaviorTreeDebugger.cpp
+++ b/Engine/Source/Editor/BehaviorTreeEditor/Private/BehaviorTreeDebugger.cpp
@@ -5,6 +5,7 @@
 #include "GameFramework/Pawn.h"
 #include "GameFramework/Controller.h"
 #include "EngineGlobals.h"
+#include "EdGraph/EdGraph.h"
 #include "Editor/EditorEngine.h"
 #include "Editor.h"
 #include "BehaviorTree/BTNode.h"
diff --git a/Engine/Source/Runtime/Engine/Private/SoundBase.cpp b/Engine/Source/Runtime/Engine/Private/SoundBase.cpp
--- a/Engine/Source/Runtime/Engine/Private/SoundBase.cpp
+++ b/Engine/Source/Runtime/Engine/Private/SoundBase.cpp
@@ -4,6 +4,7 @@
 #include "EngineDefines.h"
 #include "IAudioExtensionPlugin.h"
 #include "Sound/AudioSettings.h"
+#include "Sound/SoundClass.h"
 #include "Sound/SoundSubmix.h"
 
 
diff --git a/Engine/Source/Runtime/Engine/Private/SoundCue.cpp b/Engine/Source/Runtime/Engine/Private/SoundCue.cpp
--- a/Engine/Source/Runtime/Engine/Private/SoundCue.cpp
+++ b/Engine/Source/Runtime/Engine/Private/SoundCue.cpp
@@ -9,6 +9,7 @@
 #include "Components/AudioComponent.h"
 #include "UObject/UObjectIterator.h"
 #include "EngineUtils.h"
+#include "Sound/SoundClass.h"
 #include "Sound/SoundNode.h"
 #include "Sound/SoundNodeAssetReferencer.h"
 #include "Sound/SoundNodeMixer.h"
diff --git a/Engine/Source/Runtime/Engine/Private/SoundNodeQualityLevel.cpp b/Engine/Source/Runtime/Engine/Private/SoundNodeQualityLevel.cpp
--- a/Engine/Source/Runtime/Engine/Private/SoundNodeQualityLevel.cpp
+++ b/Engine/Source/Runtime/Engine/Private/SoundNodeQualityLevel.cpp
@@ -6,6 +6,7 @@
 #include "ActiveSound.h"
 #include "Sound/AudioSettings.h"
 #include "Sound/SoundCue.h"
+#include "Sound/SoundNode.h"
 #include "GameFramework/GameUserSettings.h"
 #include "Engine/Engine.h"
 
