BehaviorTreeDebugger: only collect current runtime node descriptions when displaying the most recent step since the current tree state will match. Otherwise we keep the descriptions from the displayed debugger step that were updated by UpdateAssetFlags.

diff --git a/Engine/Source/Editor/BehaviorTreeEditor/Private/BehaviorTreeDebugger.cpp b/Engine/Source/Editor/BehaviorTreeEditor/Private/BehaviorTreeDebugger.cpp
--- a/Engine/Source/Editor/BehaviorTreeEditor/Private/BehaviorTreeDebugger.cpp
+++ b/Engine/Source/Editor/BehaviorTreeEditor/Private/BehaviorTreeDebugger.cpp
@@ -201,12 +201,17 @@ void FBehaviorTreeDebugger::Tick(float DeltaTime)
 		{
 			UpdateAssetFlags(ShowInstance, RootNode.Get(), ActiveStepIndex);
 		}
+		
+		// Collect current runtime descriptions for every node when displaying the most recent step.
+		// Otherwise we keep the descriptions from the displayed debugger step that were updated by UpdateAssetFlags.
+		int32 LatestStepIndex = TreeInstance->DebuggerSteps.Num() - 1;
+		if (ActiveStepIndex == LatestStepIndex)
+		{
+			TArray<FString> RuntimeDescriptions;
+			TreeInstance->StoreDebuggerRuntimeValues(RuntimeDescriptions, ShowInstance.RootNode, DebuggerInstanceIndex);
 
-		// collect current runtime descriptions for every node
-		TArray<FString> RuntimeDescriptions;
-		TreeInstance->StoreDebuggerRuntimeValues(RuntimeDescriptions, ShowInstance.RootNode, DebuggerInstanceIndex);
-	
-		UpdateAssetRuntimeDescription(RuntimeDescriptions, RootNode.Get());
+			UpdateAssetRuntimeDescription(RuntimeDescriptions, RootNode.Get());
+		}
 	}
 
 	UpdateDebuggerViewOnTick();
