Call notifications when editing Blackboard data
- Call Pre/PostEditChangeProperty() on add, delete, rename key
- Update FBlackboardKeySelector references on all classes deriving from IBlackboardDataProvider when renaming a Blackboard key (used to beBehaviorTree only)

diff --git a/Engine/Source/Editor/BehaviorTreeEditor/Private/SBehaviorTreeBlackboardEditor.cpp b/Engine/Source/Editor/BehaviorTreeEditor/Private/SBehaviorTreeBlackboardEditor.cpp
--- a/Engine/Source/Editor/BehaviorTreeEditor/Private/SBehaviorTreeBlackboardEditor.cpp
+++ b/Engine/Source/Editor/BehaviorTreeEditor/Private/SBehaviorTreeBlackboardEditor.cpp
@@ -114,6 +114,9 @@ void SBehaviorTreeBlackboardEditor::HandleDeleteEntry()
 			const FScopedTransaction Transaction(LOCTEXT("BlackboardEntryDeleteTransaction", "Delete Blackboard Entry"));
 			BlackboardData->SetFlags(RF_Transactional);
 			BlackboardData->Modify();
+
+			UProperty* KeysProperty = FindField<UProperty>(UBlackboardData::StaticClass(), GET_MEMBER_NAME_CHECKED(UBlackboardData, Keys));
+			BlackboardData->PreEditChange(KeysProperty);
 		
 			for(int32 ItemIndex = 0; ItemIndex < BlackboardData->Keys.Num(); ItemIndex++)
 			{
@@ -132,6 +135,9 @@ void SBehaviorTreeBlackboardEditor::HandleDeleteEntry()
 			{
 				OnEntrySelected.Execute(nullptr, false);
 			}
+
+			FPropertyChangedEvent PropertyChangedEvent(KeysProperty, EPropertyChangeType::ArrayRemove);
+			BlackboardData->PostEditChangeProperty(PropertyChangedEvent);
 		}
 	}
 }
@@ -240,6 +246,9 @@ void SBehaviorTreeBlackboardEditor::HandleKeyClassPicked(UClass* InClass)
 	BlackboardData->SetFlags(RF_Transactional);
 	BlackboardData->Modify();
 
+	UProperty* KeysProperty = FindField<UProperty>(UBlackboardData::StaticClass(), GET_MEMBER_NAME_CHECKED(UBlackboardData, Keys));
+	BlackboardData->PreEditChange(KeysProperty);
+
 	// create a name for this new key
 	FString NewKeyName = InClass->GetDisplayNameText().ToString();
 	NewKeyName = NewKeyName.Replace(TEXT(" "), TEXT(""));
@@ -291,6 +300,9 @@ void SBehaviorTreeBlackboardEditor::HandleKeyClassPicked(UClass* InClass)
 	BlackboardEntryAction->bIsNew = true;
 
 	GraphActionMenu->OnRequestRenameOnActionNode();
+
+	FPropertyChangedEvent PropertyChangedEvent(KeysProperty, EPropertyChangeType::ArrayAdd);
+	BlackboardData->PostEditChangeProperty(PropertyChangedEvent);
 }
 
 bool SBehaviorTreeBlackboardEditor::CanCreateNewEntry() const
diff --git a/Engine/Source/Editor/BehaviorTreeEditor/Private/SBehaviorTreeBlackboardView.cpp b/Engine/Source/Editor/BehaviorTreeEditor/Private/SBehaviorTreeBlackboardView.cpp
--- a/Engine/Source/Editor/BehaviorTreeEditor/Private/SBehaviorTreeBlackboardView.cpp
+++ b/Engine/Source/Editor/BehaviorTreeEditor/Private/SBehaviorTreeBlackboardView.cpp
@@ -11,6 +11,7 @@
 #include "Modules/ModuleManager.h"
 #include "UObject/UObjectHash.h"
 #include "UObject/UnrealType.h"
+#include "UObject/UObjectIterator.h"
 #include "Widgets/SBoxPanel.h"
 #include "Framework/Commands/UIAction.h"
 #include "Framework/Commands/UICommandList.h"
@@ -23,6 +24,7 @@
 #include "SGraphActionMenu.h"
 #include "SGraphPalette.h"
 #include "BehaviorTree/BehaviorTree.h"
+#include "BehaviorTree/BlackboardAssetProvider.h"
 #include "Styling/SlateIconFinder.h"
 #include "Styling/CoreStyle.h"
 #include "ScopedTransaction.h"
@@ -217,6 +219,16 @@ private:
 			BlackboardEntryAction->BlackboardData->Modify();
 			BlackboardEntryAction->Key.EntryName = NewName;
 
+			UProperty* KeysArrayProperty = FindField<UProperty>(UBlackboardData::StaticClass(), GET_MEMBER_NAME_CHECKED(UBlackboardData, Keys));
+			UProperty* NameProperty = FindField<UProperty>(FBlackboardEntry::StaticStruct(), GET_MEMBER_NAME_CHECKED(FBlackboardEntry, EntryName));
+			FEditPropertyChain PropertyChain;
+			PropertyChain.AddHead(KeysArrayProperty);
+			PropertyChain.AddTail(NameProperty);
+			PropertyChain.SetActiveMemberPropertyNode(KeysArrayProperty);
+			PropertyChain.SetActivePropertyNode(NameProperty);
+
+			BlackboardEntryAction->BlackboardData->PreEditChange(PropertyChain);
+
 			BlackboardEntryAction->Update();
 
 			OnBlackboardKeyChanged.ExecuteIfBound(BlackboardEntryAction->BlackboardData, &BlackboardEntryAction->Key);
@@ -225,13 +237,33 @@ private:
 			{
 				UpdateExternalBlackboardKeyReferences(OldName, NewName, ExternalBTAssetsWithKeyReferences);
 			}
+
+			FPropertyChangedEvent PropertyChangedEvent(NameProperty, EPropertyChangeType::ValueSet);
+			FPropertyChangedChainEvent PropertyChangedChainEvent(PropertyChain, PropertyChangedEvent);
+			BlackboardEntryAction->BlackboardData->PostEditChangeChainProperty(PropertyChangedChainEvent);
 		}
 
 		BlackboardEntryAction->bIsNew = false;
 	}
 
+	void GetBlackboardOwnerClasses(TArray<const UClass*>& BlackboardOwnerClasses)
+	{
+		for (TObjectIterator<UClass> ClassIt; ClassIt; ++ClassIt)
+		{
+			class UClass* Class = *ClassIt;
+			if (Class->ImplementsInterface(UBlackboardAssetProvider::StaticClass()))
+			{
+				BlackboardOwnerClasses.Add(Class);
+			}
+		}
+	}
+
 	void LoadReferencerBehaviorTrees(const UBlackboardData& InBlackboardData, TArray<UObject*>& OutExternalBTAssetsWithKeyReferences)
 	{
+		// Get classes and derived classes which implement UBlackboardAssetProvider.
+		TArray<const UClass*> BlackboardOwnerClasses;
+		GetBlackboardOwnerClasses(BlackboardOwnerClasses);
+
 		IAssetRegistry& AssetRegistry = FModuleManager::LoadModuleChecked<FAssetRegistryModule>(TEXT("AssetRegistry")).Get();
 
 		TArray<FName> ReferencerPackages;
@@ -249,14 +281,15 @@ private:
 
 				for (const FAssetData& Asset : Assets)
 				{
-					if (Asset.GetClass()->IsChildOf(UBehaviorTree::StaticClass()))
+					if (BlackboardOwnerClasses.Find(Asset.GetClass()) != INDEX_NONE)
 					{
 						SlowTask.EnterProgressFrame(1.0f, FText::Format(LOCTEXT("CheckingBehaviorTree", "Key renamed, loading {0}"), FText::FromName(Asset.AssetName)));
 
-						UBehaviorTree* BehaviorTree = Cast<UBehaviorTree>(Asset.GetAsset());
-						if (BehaviorTree->BlackboardAsset == &InBlackboardData)
+						UObject* AssetObject = Asset.GetAsset();
+						const IBlackboardAssetProvider* BlackboardProvider = Cast<const IBlackboardAssetProvider>(AssetObject);
+						if (BlackboardProvider && BlackboardProvider->GetBlackboardAsset() == &InBlackboardData)
 						{
-							OutExternalBTAssetsWithKeyReferences.Add(BehaviorTree);
+							OutExternalBTAssetsWithKeyReferences.Add(AssetObject);
 						}
 					}
 				}
diff --git a/Engine/Source/Runtime/AIModule/Classes/BehaviorTree/BTNode.h b/Engine/Source/Runtime/AIModule/Classes/BehaviorTree/BTNode.h
--- a/Engine/Source/Runtime/AIModule/Classes/BehaviorTree/BTNode.h
+++ b/Engine/Source/Runtime/AIModule/Classes/BehaviorTree/BTNode.h
@@ -34,7 +34,7 @@ struct FBTInstancedNodeMemory
 };
 
 UCLASS(Abstract,config=Game)
-class AIMODULE_API UBTNode : public UObject, public IGameplayTaskOwnerInterface, public IBlackboardAssetProvider
+class AIMODULE_API UBTNode : public UObject, public IGameplayTaskOwnerInterface
 {
 	GENERATED_UCLASS_BODY()
 
@@ -133,10 +133,8 @@ class AIMODULE_API UBTNode : public UObject, public IGameplayTaskOwnerInterface,
 	/** @return tree asset */
 	UBehaviorTree* GetTreeAsset() const;
 
-	// BEGIN IBlackboardAssetProvider
 	/** @return blackboard asset */
-	virtual UBlackboardData* GetBlackboardAsset() const override;
-	// END IBlackboardAssetProvider
+	UBlackboardData* GetBlackboardAsset() const;
 
 	/** @return node instance if bCreateNodeInstance was set */
 	UBTNode* GetNodeInstance(const UBehaviorTreeComponent& OwnerComp, uint8* NodeMemory) const;
diff --git a/Engine/Source/Runtime/AIModule/Classes/BehaviorTree/BehaviorTree.h b/Engine/Source/Runtime/AIModule/Classes/BehaviorTree/BehaviorTree.h
--- a/Engine/Source/Runtime/AIModule/Classes/BehaviorTree/BehaviorTree.h
+++ b/Engine/Source/Runtime/AIModule/Classes/BehaviorTree/BehaviorTree.h
@@ -12,7 +12,7 @@ class UBlackboardData;
 class UBTDecorator;
 
 UCLASS(BlueprintType)
-class AIMODULE_API UBehaviorTree : public UObject
+class AIMODULE_API UBehaviorTree : public UObject, public IBlackboardAssetProvider
 {
 	GENERATED_UCLASS_BODY()
 
@@ -32,6 +32,11 @@ class AIMODULE_API UBehaviorTree : public UObject
 
 #endif
 
+	// BEGIN IBlackboardAssetProvider
+	/** @return blackboard asset */
+	virtual UBlackboardData* GetBlackboardAsset() const override;
+	// END IBlackboardAssetProvider
+
 	/** blackboard asset for this tree */
 	UPROPERTY()
 	UBlackboardData* BlackboardAsset;
diff --git a/Engine/Source/Runtime/AIModule/Private/BehaviorTree/BehaviorTree.cpp b/Engine/Source/Runtime/AIModule/Private/BehaviorTree/BehaviorTree.cpp
--- a/Engine/Source/Runtime/AIModule/Private/BehaviorTree/BehaviorTree.cpp
+++ b/Engine/Source/Runtime/AIModule/Private/BehaviorTree/BehaviorTree.cpp
@@ -7,3 +7,8 @@ DEFINE_LOG_CATEGORY(LogBehaviorTree);
 UBehaviorTree::UBehaviorTree(const FObjectInitializer& ObjectInitializer) : Super(ObjectInitializer)
 {
 }
+
+UBlackboardData* UBehaviorTree::GetBlackboardAsset() const
+{
+	return BlackboardAsset;
+}
