Fix deprecation warnings
#rnx

diff --git a/Engine/Source/Editor/BehaviorTreeEditor/Private/SBehaviorTreeBlackboardEditor.cpp b/Engine/Source/Editor/BehaviorTreeEditor/Private/SBehaviorTreeBlackboardEditor.cpp
--- a/Engine/Source/Editor/BehaviorTreeEditor/Private/SBehaviorTreeBlackboardEditor.cpp
+++ b/Engine/Source/Editor/BehaviorTreeEditor/Private/SBehaviorTreeBlackboardEditor.cpp
@@ -115,7 +115,7 @@ void SBehaviorTreeBlackboardEditor::HandleDeleteEntry()
 			BlackboardData->SetFlags(RF_Transactional);
 			BlackboardData->Modify();
 
-			UProperty* KeysProperty = FindField<UProperty>(UBlackboardData::StaticClass(), GET_MEMBER_NAME_CHECKED(UBlackboardData, Keys));
+			FProperty* KeysProperty = FindField<FProperty>(UBlackboardData::StaticClass(), GET_MEMBER_NAME_CHECKED(UBlackboardData, Keys));
 			BlackboardData->PreEditChange(KeysProperty);
 		
 			for(int32 ItemIndex = 0; ItemIndex < BlackboardData->Keys.Num(); ItemIndex++)
@@ -246,7 +246,7 @@ void SBehaviorTreeBlackboardEditor::HandleKeyClassPicked(UClass* InClass)
 	BlackboardData->SetFlags(RF_Transactional);
 	BlackboardData->Modify();
 
-	UProperty* KeysProperty = FindField<UProperty>(UBlackboardData::StaticClass(), GET_MEMBER_NAME_CHECKED(UBlackboardData, Keys));
+	FProperty* KeysProperty = FindField<FProperty>(UBlackboardData::StaticClass(), GET_MEMBER_NAME_CHECKED(UBlackboardData, Keys));
 	BlackboardData->PreEditChange(KeysProperty);
 
 	// create a name for this new key
diff --git a/Engine/Source/Editor/BehaviorTreeEditor/Private/SBehaviorTreeBlackboardView.cpp b/Engine/Source/Editor/BehaviorTreeEditor/Private/SBehaviorTreeBlackboardView.cpp
--- a/Engine/Source/Editor/BehaviorTreeEditor/Private/SBehaviorTreeBlackboardView.cpp
+++ b/Engine/Source/Editor/BehaviorTreeEditor/Private/SBehaviorTreeBlackboardView.cpp
@@ -219,8 +219,8 @@ private:
 			BlackboardEntryAction->BlackboardData->Modify();
 			BlackboardEntryAction->Key.EntryName = NewName;
 
-			UProperty* KeysArrayProperty = FindField<UProperty>(UBlackboardData::StaticClass(), GET_MEMBER_NAME_CHECKED(UBlackboardData, Keys));
-			UProperty* NameProperty = FindField<UProperty>(FBlackboardEntry::StaticStruct(), GET_MEMBER_NAME_CHECKED(FBlackboardEntry, EntryName));
+			FProperty* KeysArrayProperty = FindField<FProperty>(UBlackboardData::StaticClass(), GET_MEMBER_NAME_CHECKED(UBlackboardData, Keys));
+			FProperty* NameProperty = FindField<FProperty>(FBlackboardEntry::StaticStruct(), GET_MEMBER_NAME_CHECKED(FBlackboardEntry, EntryName));
 			FEditPropertyChain PropertyChain;
 			PropertyChain.AddHead(KeysArrayProperty);
 			PropertyChain.AddTail(NameProperty);
