BehaviorTreeDebugger: prevent displaying old subtree recorded step. Fixed displayed step index using recording buffer index instead of debugger execution step id.

diff --git a/Engine/Source/Editor/BehaviorTreeEditor/Private/BehaviorTreeDebugger.cpp b/Engine/Source/Editor/BehaviorTreeEditor/Private/BehaviorTreeDebugger.cpp
--- a/Engine/Source/Editor/BehaviorTreeEditor/Private/BehaviorTreeDebugger.cpp
+++ b/Engine/Source/Editor/BehaviorTreeEditor/Private/BehaviorTreeDebugger.cpp
@@ -84,7 +84,7 @@ void FBehaviorTreeDebugger::Setup(UBehaviorTree* InTreeAsset, TSharedRef<FBehavi
 	TreeAsset = InTreeAsset;
 	DebuggerInstanceIndex = INDEX_NONE;
 	ActiveStepIndex = 0;
-	LastValidStepId = INDEX_NONE;
+	LastValidExecutionStepId = FBehaviorTreeExecutionStep::InvalidExecutionId;
 	ActiveBreakpoints.Reset();
 	KnownInstances.Reset();
 
@@ -119,12 +119,13 @@ void FBehaviorTreeDebugger::Refresh()
 				UpdateDebuggerViewOnStepChange();
 				UpdateDebuggerViewOnTick();
 
-				const FBehaviorTreeDebuggerInstance& ShowInstance = TreeInstance->DebuggerSteps[ActiveStepIndex].InstanceStack[DebuggerInstanceIndex];
+				FBehaviorTreeExecutionStep& ExecutionStep = TreeInstance->DebuggerSteps[ActiveStepIndex];
+				const FBehaviorTreeDebuggerInstance& ShowInstance = ExecutionStep.InstanceStack[DebuggerInstanceIndex];
 				OnActiveNodeChanged(ShowInstance.ActivePath, HasContinuousPrevStep() ?
 					TreeInstance->DebuggerSteps[ActiveStepIndex - 1].InstanceStack[DebuggerInstanceIndex].ActivePath :
 					TArray<uint16>());
 
-				UpdateAssetFlags(ShowInstance, RootNode.Get(), ActiveStepIndex);
+				UpdateAssetFlags(ShowInstance, RootNode.Get(), ExecutionStep.ExecutionStepId);
 			}
 		}
 	}
@@ -152,28 +153,40 @@ void FBehaviorTreeDebugger::Tick(float DeltaTime)
 	TArray<uint16> EmptyPath;
 
 	int32 TestStepIndex = 0;
-	for (int32 Idx = TreeInstance->DebuggerSteps.Num() - 1; Idx >= 0; Idx--)
+
+	if (LastValidExecutionStepId != FBehaviorTreeExecutionStep::InvalidExecutionId)
 	{
-		const FBehaviorTreeExecutionStep& Step = TreeInstance->DebuggerSteps[Idx];
-		if (Step.StepIndex == LastValidStepId)
+		for (int32 Idx = TreeInstance->DebuggerSteps.Num() - 1; Idx >= 0; Idx--)
 		{
-			TestStepIndex = Idx;
-			break;
+			const FBehaviorTreeExecutionStep& Step = TreeInstance->DebuggerSteps[Idx];
+			if (Step.ExecutionStepId == LastValidExecutionStepId)
+			{
+				TestStepIndex = Idx;
+				break;
+			}
 		}
 	}
 
+	// Keep track if an update of the debugger instance is required after processing changes from displayed step to latest step
+	// since it is possible the current instance gets invalidated (e.g. subtree) and that resets the current index. In such case
+	// we need to wait for the next tick to recompute ActiveStepIndex and LastValidExecutionStepId.
+	bool bDebuggerInstanceUpdateRequired = true;
+
 	// find index of previously displayed state and notify about all changes in between to give breakpoints a chance to trigger
 	for (int32 i = TestStepIndex; i < TreeInstance->DebuggerSteps.Num(); i++)
 	{
 		const FBehaviorTreeExecutionStep& Step = TreeInstance->DebuggerSteps[i];
-		if (Step.StepIndex > DisplayedStepIndex)
+
+		if (Step.ExecutionStepId > DisplayedExecutionStepId)
 		{
 			ActiveStepIndex = i;
-			LastValidStepId = Step.StepIndex;
+			LastValidExecutionStepId = Step.ExecutionStepId;
 
 			UpdateDebuggerInstance();
 			UpdateAvailableActions();
 
+			bDebuggerInstanceUpdateRequired = false;
+
 			if (DebuggerInstanceIndex != INDEX_NONE)
 			{
 				UpdateDebuggerViewOnStepChange();
@@ -192,14 +205,19 @@ void FBehaviorTreeDebugger::Tick(float DeltaTime)
 		}
 	}
 
-	UpdateDebuggerInstance();
+	if (bDebuggerInstanceUpdateRequired)
+	{
+		UpdateDebuggerInstance();
+	}
+
 	if (DebuggerInstanceIndex != INDEX_NONE)
 	{
-		const FBehaviorTreeDebuggerInstance& ShowInstance = TreeInstance->DebuggerSteps[ActiveStepIndex].InstanceStack[DebuggerInstanceIndex];
+		FBehaviorTreeExecutionStep& ExecutionStep = TreeInstance->DebuggerSteps[ActiveStepIndex];
+		const FBehaviorTreeDebuggerInstance& ShowInstance = ExecutionStep.InstanceStack[DebuggerInstanceIndex];
 
-		if (DisplayedStepIndex != TreeInstance->DebuggerSteps[ActiveStepIndex].StepIndex)
+		if (DisplayedExecutionStepId != ExecutionStep.ExecutionStepId)
 		{
-			UpdateAssetFlags(ShowInstance, RootNode.Get(), ActiveStepIndex);
+			UpdateAssetFlags(ShowInstance, RootNode.Get(), ExecutionStep.ExecutionStepId);
 		}
 		
 		// Collect current runtime descriptions for every node when displaying the most recent step.
@@ -264,7 +282,7 @@ void FBehaviorTreeDebugger::OnEndPIE(const bool bIsSimulating)
 	ActiveBreakpoints.Reset();
 
 	FBehaviorTreeDebuggerInstance EmptyData;
-	UpdateAssetFlags(EmptyData, RootNode.Get(), INDEX_NONE);
+	UpdateAssetFlags(EmptyData, RootNode.Get(), FBehaviorTreeExecutionStep::InvalidExecutionId);
 	UpdateDebuggerViewOnInstanceChange();
 }
 
@@ -330,16 +348,16 @@ void FBehaviorTreeDebugger::OnTreeStarted(const UBehaviorTreeComponent& OwnerCom
 
 void FBehaviorTreeDebugger::ClearDebuggerState(bool bKeepSubtree)
 {
-	LastValidStepId = bKeepSubtree ? LastValidStepId : INDEX_NONE;
+	LastValidExecutionStepId = bKeepSubtree ? LastValidExecutionStepId : FBehaviorTreeExecutionStep::InvalidExecutionId;
+	DisplayedExecutionStepId = FBehaviorTreeExecutionStep::InvalidExecutionId;
 
 	DebuggerInstanceIndex = INDEX_NONE;
 	ActiveStepIndex = 0;
-	DisplayedStepIndex = INDEX_NONE;
 
 	if (TreeAsset && RootNode.IsValid())
 	{
 		FBehaviorTreeDebuggerInstance EmptyData;
-		UpdateAssetFlags(EmptyData, RootNode.Get(), INDEX_NONE);
+		UpdateAssetFlags(EmptyData, RootNode.Get(), FBehaviorTreeExecutionStep::InvalidExecutionId);
 	}
 }
 
@@ -380,7 +398,7 @@ void FBehaviorTreeDebugger::SetNodeFlags(const struct FBehaviorTreeDebuggerInsta
 	const bool bIsNodeActive = bIsNodeActivePath || bIsNodeActiveAdditional;
 	const bool bIsShowingCurrentState = IsShowingCurrentState();
 
-	Node->DebuggerUpdateCounter = DisplayedStepIndex;
+	Node->DebuggerUpdateCounter = DisplayedExecutionStepId;
 	Node->bDebuggerMarkCurrentlyActive = bIsNodeActive && bIsShowingCurrentState;
 	Node->bDebuggerMarkPreviouslyActive = bIsNodeActive && !bIsShowingCurrentState;
 	
@@ -442,7 +460,7 @@ void FBehaviorTreeDebugger::SetCompositeDecoratorFlags(const struct FBehaviorTre
 		}
 	}	
 	
-	Node->DebuggerUpdateCounter = DisplayedStepIndex;
+	Node->DebuggerUpdateCounter = DisplayedExecutionStepId;
 	Node->bDebuggerMarkCurrentlyActive = bIsNodeActive && bIsShowingCurrentState;
 	Node->bDebuggerMarkPreviouslyActive = bIsNodeActive && !bIsShowingCurrentState;
 
@@ -480,7 +498,7 @@ void FBehaviorTreeDebugger::SetCompositeDecoratorFlags(const struct FBehaviorTre
 	Node->DebuggerSearchPathSize = Data.PathFromPrevious.Num() - NumTriggers;
 }
 
-void FBehaviorTreeDebugger::UpdateAssetFlags(const struct FBehaviorTreeDebuggerInstance& Data, class UBehaviorTreeGraphNode* Node, int32 StepIdx)
+void FBehaviorTreeDebugger::UpdateAssetFlags(const struct FBehaviorTreeDebuggerInstance& Data, class UBehaviorTreeGraphNode* Node, int32 ExecutionStepId)
 {
 	if (Node == NULL)
 	{
@@ -490,12 +508,12 @@ void FBehaviorTreeDebugger::UpdateAssetFlags(const struct FBehaviorTreeDebuggerI
 	// special case for marking root when out of nodes
 	if (Node == RootNode.Get())
 	{
-		const bool bIsNodeActive = (Data.ActivePath.Num() == 0) && (StepIdx >= 0);
+		const bool bIsNodeActive = (Data.ActivePath.Num() == 0) && (ExecutionStepId != FBehaviorTreeExecutionStep::InvalidExecutionId);
 		const bool bIsShowingCurrentState = IsShowingCurrentState();
 
 		Node->bDebuggerMarkCurrentlyActive = bIsNodeActive && bIsShowingCurrentState;
 		Node->bDebuggerMarkPreviouslyActive = bIsNodeActive && !bIsShowingCurrentState;
-		DisplayedStepIndex = StepIdx;
+		DisplayedExecutionStepId = ExecutionStepId;
 	}
 
 	for (int32 PinIdx = 0; PinIdx < Node->Pins.Num(); PinIdx++)
@@ -555,7 +573,7 @@ void FBehaviorTreeDebugger::UpdateAssetFlags(const struct FBehaviorTreeDebuggerI
 					}
 				}
 
-				UpdateAssetFlags(Data, LinkedNode, StepIdx);
+				UpdateAssetFlags(Data, LinkedNode, ExecutionStepId);
 			}
 		}
 	}
@@ -946,14 +964,14 @@ void FBehaviorTreeDebugger::UpdateCurrentStep(int32 PrevStepIdx, int32 NewStepId
 		if (NewStepInfo.InstanceStack.IsValidIndex(DebuggerInstanceIndex))
 		{
 			const FBehaviorTreeDebuggerInstance& ShowInstance = NewStepInfo.InstanceStack[DebuggerInstanceIndex];
-			UpdateAssetFlags(ShowInstance, RootNode.Get(), ActiveStepIndex);
+			UpdateAssetFlags(ShowInstance, RootNode.Get(), NewStepInfo.ExecutionStepId);
 		}
 		else
 		{
 			ActiveStepIndex = INDEX_NONE;
 
 			FBehaviorTreeDebuggerInstance EmptyData;
-			UpdateAssetFlags(EmptyData, RootNode.Get(), INDEX_NONE);
+			UpdateAssetFlags(EmptyData, RootNode.Get(), FBehaviorTreeExecutionStep::InvalidExecutionId);
 		}
 
 		UpdateDebuggerViewOnStepChange();
@@ -1002,6 +1020,11 @@ bool FBehaviorTreeDebugger::HasContinuousPrevStep() const
 
 void FBehaviorTreeDebugger::OnActiveNodeChanged(const TArray<uint16>& ActivePath, const TArray<uint16>& PrevStepPath)
 {
+	if (ActiveBreakpoints.Num() == 0)
+	{
+		return;
+	}
+
 	bool bShouldPause = false;
 	StoppedOnBreakpointExecutionIndex = MAX_uint16;
 	
@@ -1306,11 +1329,14 @@ void FBehaviorTreeDebugger::InitializeFromParent(class FBehaviorTreeDebugger* Pa
 	UpdateAvailableActions();
 
 	if (TreeInstance.IsValid() &&
-		TreeInstance->DebuggerSteps.IsValidIndex(ActiveStepIndex) &&
-		TreeInstance->DebuggerSteps[ActiveStepIndex].InstanceStack.IsValidIndex(DebuggerInstanceIndex))
+		TreeInstance->DebuggerSteps.IsValidIndex(ActiveStepIndex))
 	{
-		const FBehaviorTreeDebuggerInstance& ShowInstance = TreeInstance->DebuggerSteps[ActiveStepIndex].InstanceStack[DebuggerInstanceIndex];
-		UpdateAssetFlags(ShowInstance, RootNode.Get(), ActiveStepIndex);
+		FBehaviorTreeExecutionStep& ExecutionStep = TreeInstance->DebuggerSteps[ActiveStepIndex];
+		if (ExecutionStep.InstanceStack.IsValidIndex(DebuggerInstanceIndex))
+		{
+			const FBehaviorTreeDebuggerInstance& ShowInstance = ExecutionStep.InstanceStack[DebuggerInstanceIndex];
+			UpdateAssetFlags(ShowInstance, RootNode.Get(), ExecutionStep.ExecutionStepId);
+		}
 	}
 #endif
 }
diff --git a/Engine/Source/Editor/BehaviorTreeEditor/Private/BehaviorTreeDebugger.h b/Engine/Source/Editor/BehaviorTreeEditor/Private/BehaviorTreeDebugger.h
--- a/Engine/Source/Editor/BehaviorTreeEditor/Private/BehaviorTreeDebugger.h
+++ b/Engine/Source/Editor/BehaviorTreeEditor/Private/BehaviorTreeDebugger.h
@@ -112,11 +112,11 @@ private:
 	/** index of state from buffer to show */
 	int32 ActiveStepIndex;
 
-	/** index of displayed step, used to detect changes */
-	int32 DisplayedStepIndex;
+	/** execution id of displayed step, used to detect changes */
+	int32 DisplayedExecutionStepId;
 
-	/** id of last valid step */
-	int32 LastValidStepId;
+	/** execution id of last valid step */
+	int32 LastValidExecutionStepId;
 
 	/** indices of display steps for different step actions */
 	int32 StepForwardIntoIdx;
diff --git a/Engine/Source/Runtime/AIModule/Classes/BehaviorTree/BehaviorTreeTypes.h b/Engine/Source/Runtime/AIModule/Classes/BehaviorTree/BehaviorTreeTypes.h
--- a/Engine/Source/Runtime/AIModule/Classes/BehaviorTree/BehaviorTreeTypes.h
+++ b/Engine/Source/Runtime/AIModule/Classes/BehaviorTree/BehaviorTreeTypes.h
@@ -247,7 +247,7 @@ struct FBehaviorTreeDebuggerInstance
 /** debugger data about current execution step */
 struct FBehaviorTreeExecutionStep
 {
-	FBehaviorTreeExecutionStep() : TimeStamp(0.f), StepIndex(INDEX_NONE) {}
+	FBehaviorTreeExecutionStep() : TimeStamp(0.f), ExecutionStepId(InvalidExecutionId) {}
 
 	/** subtree instance stack */
 	TArray<FBehaviorTreeDebuggerInstance> InstanceStack;
@@ -258,8 +258,10 @@ struct FBehaviorTreeExecutionStep
 	/** Game world's time stamp of this step */
 	float TimeStamp;
 
-	/** index of execution step */
-	int32 StepIndex;
+	static constexpr int32 InvalidExecutionId = -1;
+
+	/** Id of execution step */
+	int32 ExecutionStepId;
 };
 
 /** identifier of subtree instance */
diff --git a/Engine/Source/Runtime/AIModule/Private/BehaviorTree/BehaviorTreeComponent.cpp b/Engine/Source/Runtime/AIModule/Private/BehaviorTree/BehaviorTreeComponent.cpp
--- a/Engine/Source/Runtime/AIModule/Private/BehaviorTree/BehaviorTreeComponent.cpp
+++ b/Engine/Source/Runtime/AIModule/Private/BehaviorTree/BehaviorTreeComponent.cpp
@@ -2507,7 +2507,7 @@ void UBehaviorTreeComponent::StoreDebuggerExecutionStep(EBTExecutionSnap::Type S
 	}
 
 	FBehaviorTreeExecutionStep CurrentStep;
-	CurrentStep.StepIndex = DebuggerSteps.Num() ? DebuggerSteps.Last().StepIndex + 1 : 0;
+	CurrentStep.ExecutionStepId = DebuggerSteps.Num() ? DebuggerSteps.Last().ExecutionStepId + 1 : 0;
 	CurrentStep.TimeStamp = GetWorld()->GetTimeSeconds();
 	CurrentStep.BlackboardValues = SearchStartBlackboard;
 
