BehaviorTreeDebugger: added missing initialization for StoppedOnBreakpointExecutionIndex and make sure to always reset it when active node changes

#rnx

diff --git a/Engine/Source/Editor/BehaviorTreeEditor/Private/BehaviorTreeDebugger.cpp b/Engine/Source/Editor/BehaviorTreeEditor/Private/BehaviorTreeDebugger.cpp
--- a/Engine/Source/Editor/BehaviorTreeEditor/Private/BehaviorTreeDebugger.cpp
+++ b/Engine/Source/Editor/BehaviorTreeEditor/Private/BehaviorTreeDebugger.cpp
@@ -29,6 +29,7 @@ FBehaviorTreeDebugger::FBehaviorTreeDebugger()
 	TreeAsset = NULL;
 	bIsPIEActive = false;
 	bIsCurrentSubtree = false;
+	StoppedOnBreakpointExecutionIndex = INDEX_NONE;
 	StepForwardIntoIdx = INDEX_NONE;
 	StepForwardOverIdx = INDEX_NONE;
 	StepBackIntoIdx = INDEX_NONE;
@@ -1020,13 +1021,13 @@ bool FBehaviorTreeDebugger::HasContinuousPrevStep() const
 
 void FBehaviorTreeDebugger::OnActiveNodeChanged(const TArray<uint16>& ActivePath, const TArray<uint16>& PrevStepPath)
 {
+	StoppedOnBreakpointExecutionIndex = INDEX_NONE;
 	if (ActiveBreakpoints.Num() == 0)
 	{
 		return;
 	}
 
 	bool bShouldPause = false;
-	StoppedOnBreakpointExecutionIndex = MAX_uint16;
 	
 	// breakpoints: check only nodes, that have changed from previous state
 	// (e.g. breakpoint on sequence, it would break multiple times for every child
