Fix BehaviorTreeEditor auto arrange crash and bug.

diff --git a/Engine/Source/Editor/BehaviorTreeEditor/Private/BehaviorTreeGraph.cpp b/Engine/Source/Editor/BehaviorTreeEditor/Private/BehaviorTreeGraph.cpp
--- a/Engine/Source/Editor/BehaviorTreeEditor/Private/BehaviorTreeGraph.cpp
+++ b/Engine/Source/Editor/BehaviorTreeEditor/Private/BehaviorTreeGraph.cpp
@@ -1212,17 +1212,18 @@ namespace BTAutoArrangeHelpers
 		if (Pin)
 		{
 			SGraphNode::FNodeSet NodeFilter;
-			for (int32 Idx = 0; Idx < Pin->LinkedTo.Num(); Idx++)
+			TArray<UEdGraphPin*> TempLinkedTo = Pin->LinkedTo;
+			for (int32 Idx = 0; Idx < TempLinkedTo.Num(); Idx++)
 			{
-				UBehaviorTreeGraphNode* GraphNode = Cast<UBehaviorTreeGraphNode>(Pin->LinkedTo[Idx]->GetOwningNode());
+				UBehaviorTreeGraphNode* GraphNode = Cast<UBehaviorTreeGraphNode>(TempLinkedTo[Idx]->GetOwningNode());
 				if (GraphNode && BBoxTree.Children.Num() > 0)
 				{
 					AutoArrangeNodes(GraphNode, BBoxTree.Children[BBoxIndex], PosX, PosY + GraphNode->DEPRECATED_NodeWidget.Pin()->GetDesiredSize().Y * 2.5f);
 					GraphNode->DEPRECATED_NodeWidget.Pin()->MoveTo(FVector2D(BBoxTree.Children[BBoxIndex].SubGraphBBox.X / 2 - GraphNode->DEPRECATED_NodeWidget.Pin()->GetDesiredSize().X / 2 + PosX, PosY), NodeFilter);
 					PosX += BBoxTree.Children[BBoxIndex].SubGraphBBox.X + 20;
+					BBoxIndex++;
 				}
 
-				BBoxIndex++;
 			}
 		}
 	}
