Fixing C6286 code analysis (MSVC) error in SGraphNode_BehaviorTree.cpp

diff --git a/Engine/Source/Editor/BehaviorTreeEditor/Private/SGraphNode_BehaviorTree.cpp b/Engine/Source/Editor/BehaviorTreeEditor/Private/SGraphNode_BehaviorTree.cpp
--- a/Engine/Source/Editor/BehaviorTreeEditor/Private/SGraphNode_BehaviorTree.cpp
+++ b/Engine/Source/Editor/BehaviorTreeEditor/Private/SGraphNode_BehaviorTree.cpp
@@ -33,9 +33,15 @@
 
 #define LOCTEXT_NAMESPACE "BehaviorTreeEditor"
 
-namespace
+#define ALWAYS_SHOW_BT_EXECUTION_INDEX 1
+
+bool ShouldShowExecutionIndex()
 {
-	static const bool bShowExecutionIndexInEditorMode = true;
+#ifdef ALWAYS_SHOW_BT_EXECUTION_INDEX
+	return true;
+#else
+	return GEditor && (GEditor->bIsSimulatingInEditor || GEditor->PlayWorld != NULL);
+#endif // ALWAYS_SHOW_BT_EXECUTION_INDEX
 }
 
 /////////////////////////////////////////////////////
@@ -906,8 +912,7 @@ EVisibility SGraphNode_BehaviorTree::GetIndexVisibility() const
 	}
 
 	// Visible if we are in PIE or if we have siblings
-	CA_SUPPRESS(6235);
-	const bool bCanShowIndex = (bShowExecutionIndexInEditorMode || GEditor->bIsSimulatingInEditor || GEditor->PlayWorld != NULL) || (MyParentOutputPin && MyParentOutputPin->LinkedTo.Num() > 1);
+	const bool bCanShowIndex = ShouldShowExecutionIndex() || (MyParentOutputPin && MyParentOutputPin->LinkedTo.Num() > 1);
 
 	// LOD this out once things get too small
 	TSharedPtr<SGraphPanel> MyOwnerPanel = GetOwnerPanel();
@@ -926,8 +931,7 @@ FText SGraphNode_BehaviorTree::GetIndexText() const
 
 	int32 Index = 0;
 
-	CA_SUPPRESS(6235);
-	if (bShowExecutionIndexInEditorMode || GEditor->bIsSimulatingInEditor || GEditor->PlayWorld != NULL)
+	if (ShouldShowExecutionIndex())
 	{
 		// special case: range of execution indices in composite decorator node
 		UBehaviorTreeGraphNode_CompositeDecorator* CompDecorator = Cast<UBehaviorTreeGraphNode_CompositeDecorator>(GraphNode);
@@ -960,8 +964,7 @@ FText SGraphNode_BehaviorTree::GetIndexText() const
 
 FText SGraphNode_BehaviorTree::GetIndexTooltipText() const
 {
-	CA_SUPPRESS(6235);
-	if (bShowExecutionIndexInEditorMode || GEditor->bIsSimulatingInEditor || GEditor->PlayWorld != NULL)
+	if (ShouldShowExecutionIndex())
 	{
 		return LOCTEXT("ExecutionIndexTooltip", "Execution index: this shows the order in which nodes are executed.");
 	}
