PR #8029: [Blackboard editor] Fixed memory leak in details panel for blackboard selector properties (Contributed by [REDACTED_PERSON])

diff --git a/Engine/Source/Editor/BehaviorTreeEditor/Private/DetailCustomizations/BlackboardSelectorDetails.cpp b/Engine/Source/Editor/BehaviorTreeEditor/Private/DetailCustomizations/BlackboardSelectorDetails.cpp
--- a/Engine/Source/Editor/BehaviorTreeEditor/Private/DetailCustomizations/BlackboardSelectorDetails.cpp
+++ b/Engine/Source/Editor/BehaviorTreeEditor/Private/DetailCustomizations/BlackboardSelectorDetails.cpp
@@ -168,11 +168,11 @@ void FBlackboardSelectorDetails::CacheBlackboardData()
 
 	if (!OnBlackboardDataChangedHandle.IsValid())
 	{
-		UBlackboardData::OnBlackboardDataChanged.AddSP(this, &FBlackboardSelectorDetails::OnBlackboardDataChanged);
+		OnBlackboardDataChangedHandle = UBlackboardData::OnBlackboardDataChanged.AddSP(this, &FBlackboardSelectorDetails::OnBlackboardDataChanged);
 	}
 	if (!OnBlackboardOwnerChangedHandle.IsValid())
 	{
-		IBlackboardAssetProvider::OnBlackboardOwnerChanged.AddSP(this, &FBlackboardSelectorDetails::OnBlackboardOwnerChanged);
+		OnBlackboardOwnerChangedHandle = IBlackboardAssetProvider::OnBlackboardOwnerChanged.AddSP(this, &FBlackboardSelectorDetails::OnBlackboardOwnerChanged);
 	}
 
 }
