AI: Changed it so that creating new Tasks/Behaviors/Services will ask the user to choose the name/location first before creating the asset instead of requiring the user to find it in the content browser and rename it afterwards. This code was also ported to the Conversation Editor which uses a similar workflow.

diff --git a/Engine/Plugins/Experimental/CommonConversation/Source/CommonConversationEditor/Private/ConversationEditor.cpp b/Engine/Plugins/Experimental/CommonConversation/Source/CommonConversationEditor/Private/ConversationEditor.cpp
--- a/Engine/Plugins/Experimental/CommonConversation/Source/CommonConversationEditor/Private/ConversationEditor.cpp
+++ b/Engine/Plugins/Experimental/CommonConversation/Source/CommonConversationEditor/Private/ConversationEditor.cpp
@@ -936,30 +936,41 @@ void FConversationEditor::HandleNewNodeClassPicked(UClass* InClass) const
 {
 	if (ConversationAsset != nullptr)
 	{
-		FString ClassName = FBlueprintEditorUtils::GetClassNameWithoutSuffix(InClass);
+		const FString ClassName = FBlueprintEditorUtils::GetClassNameWithoutSuffix(InClass);
 
 		FString PathName = ConversationAsset->GetOutermost()->GetPathName();
 		PathName = FPaths::GetPath(PathName);
-		PathName /= ClassName;
 
-		FString Name;
-		FString PackageName;
-		FAssetToolsModule& AssetToolsModule = FModuleManager::GetModuleChecked<FAssetToolsModule>("AssetTools");
-		AssetToolsModule.Get().CreateUniqueAssetName(PathName, TEXT("_New"), PackageName, Name);
-
-		UPackage* Package = CreatePackage(*PackageName);
-		if (ensure(Package))
+		// Now that we've generated some reasonable default locations/names for the package, allow the user to have the final say
+		// before we create the package and initialize the blueprint inside of it.
+		FSaveAssetDialogConfig SaveAssetDialogConfig;
+		SaveAssetDialogConfig.DialogTitleOverride = LOCTEXT("SaveAssetDialogTitle", "Save Asset As");
+		SaveAssetDialogConfig.DefaultPath = PathName;
+		SaveAssetDialogConfig.DefaultAssetName = ClassName + TEXT("_New");
+		SaveAssetDialogConfig.ExistingAssetPolicy = ESaveAssetDialogExistingAssetPolicy::Disallow;
+
+		const FContentBrowserModule& ContentBrowserModule = FModuleManager::LoadModuleChecked<FContentBrowserModule>("ContentBrowser");
+		const FString SaveObjectPath = ContentBrowserModule.Get().CreateModalSaveAssetDialog(SaveAssetDialogConfig);
+		if (!SaveObjectPath.IsEmpty())
 		{
-			// Create and init a new Blueprint
-			if (UBlueprint* NewBP = FKismetEditorUtilities::CreateBlueprint(InClass, Package, FName(*Name), BPTYPE_Normal, UBlueprint::StaticClass(), UBlueprintGeneratedClass::StaticClass()))
+			const FString SavePackageName = FPackageName::ObjectPathToPackageName(SaveObjectPath);
+			const FString SavePackagePath = FPaths::GetPath(SavePackageName);
+			const FString SaveAssetName = FPaths::GetBaseFilename(SavePackageName);
+
+			UPackage* Package = CreatePackage(*SavePackageName);
+			if (ensure(Package))
 			{
-				GEditor->GetEditorSubsystem<UAssetEditorSubsystem>()->OpenEditorForAsset(NewBP);
+				// Create and init a new Blueprint
+				if (UBlueprint* NewBP = FKismetEditorUtilities::CreateBlueprint(InClass, Package, FName(*SaveAssetName), BPTYPE_Normal, UBlueprint::StaticClass(), UBlueprintGeneratedClass::StaticClass()))
+				{
+					GEditor->GetEditorSubsystem<UAssetEditorSubsystem>()->OpenEditorForAsset(NewBP);
 
-				// Notify the asset registry
-				FAssetRegistryModule::AssetCreated(NewBP);
+					// Notify the asset registry
+					FAssetRegistryModule::AssetCreated(NewBP);
 
-				// Mark the package dirty...
-				Package->MarkPackageDirty();
+					// Mark the package dirty...
+					Package->MarkPackageDirty();
+				}
 			}
 		}
 	}
diff --git a/Engine/Source/Editor/BehaviorTreeEditor/Private/BehaviorTreeEditor.cpp b/Engine/Source/Editor/BehaviorTreeEditor/Private/BehaviorTreeEditor.cpp
--- a/Engine/Source/Editor/BehaviorTreeEditor/Private/BehaviorTreeEditor.cpp
+++ b/Engine/Source/Editor/BehaviorTreeEditor/Private/BehaviorTreeEditor.cpp
@@ -1722,7 +1722,6 @@ TSharedRef<SWidget> FBehaviorTreeEditor::HandleCreateNewServiceMenu() const
 
 	return FModuleManager::LoadModuleChecked<FClassViewerModule>("ClassViewer").CreateClassViewer(Options, OnPicked);
 }
-
 void FBehaviorTreeEditor::HandleNewNodeClassPicked(UClass* InClass) const
 {
 	UE_CLOG(InClass == nullptr, LogBehaviorTreeEditor, Error, TEXT("Trying to handle new node of NULL class for Behavior Treee %s ")
@@ -1730,30 +1729,41 @@ void FBehaviorTreeEditor::HandleNewNodeClassPicked(UClass* InClass) const
 
 	if(BehaviorTree != nullptr && InClass != nullptr && BehaviorTree->GetOutermost())
 	{
-		FString ClassName = FBlueprintEditorUtils::GetClassNameWithoutSuffix(InClass);
+		const FString ClassName = FBlueprintEditorUtils::GetClassNameWithoutSuffix(InClass);
 
 		FString PathName = BehaviorTree->GetOutermost()->GetPathName();
 		PathName = FPaths::GetPath(PathName);
-		PathName /= ClassName;
-
-		FString Name;
-		FString PackageName;
-		FAssetToolsModule& AssetToolsModule = FModuleManager::GetModuleChecked<FAssetToolsModule>("AssetTools");
-		AssetToolsModule.Get().CreateUniqueAssetName(PathName, TEXT("_New"), PackageName, Name);
-
-		UPackage* Package = CreatePackage( *PackageName);
-		if (ensure(Package))
+		
+		// Now that we've generated some reasonable default locations/names for the package, allow the user to have the final say
+		// before we create the package and initialize the blueprint inside of it.
+		FSaveAssetDialogConfig SaveAssetDialogConfig;
+		SaveAssetDialogConfig.DialogTitleOverride = LOCTEXT("SaveAssetDialogTitle", "Save Asset As");
+		SaveAssetDialogConfig.DefaultPath = PathName;
+		SaveAssetDialogConfig.DefaultAssetName = ClassName + TEXT("_New");
+		SaveAssetDialogConfig.ExistingAssetPolicy = ESaveAssetDialogExistingAssetPolicy::Disallow;
+
+		const FContentBrowserModule& ContentBrowserModule = FModuleManager::LoadModuleChecked<FContentBrowserModule>("ContentBrowser");
+		const FString SaveObjectPath = ContentBrowserModule.Get().CreateModalSaveAssetDialog(SaveAssetDialogConfig);
+		if (!SaveObjectPath.IsEmpty())
 		{
-			// Create and init a new Blueprint
-			if (UBlueprint* NewBP = FKismetEditorUtilities::CreateBlueprint(InClass, Package, FName(*Name), BPTYPE_Normal, UBlueprint::StaticClass(), UBlueprintGeneratedClass::StaticClass()))
+			const FString SavePackageName = FPackageName::ObjectPathToPackageName(SaveObjectPath);
+			const FString SavePackagePath = FPaths::GetPath(SavePackageName);
+			const FString SaveAssetName = FPaths::GetBaseFilename(SavePackageName);
+
+			UPackage* Package = CreatePackage(*SavePackageName);
+			if (ensure(Package))
 			{
-				GEditor->GetEditorSubsystem<UAssetEditorSubsystem>()->OpenEditorForAsset(NewBP);
+				// Create and init a new Blueprint
+				if (UBlueprint* NewBP = FKismetEditorUtilities::CreateBlueprint(InClass, Package, FName(*SaveAssetName), BPTYPE_Normal, UBlueprint::StaticClass(), UBlueprintGeneratedClass::StaticClass()))
+				{
+					GEditor->GetEditorSubsystem<UAssetEditorSubsystem>()->OpenEditorForAsset(NewBP);
 
-				// Notify the asset registry
-				FAssetRegistryModule::AssetCreated(NewBP);
+					// Notify the asset registry
+					FAssetRegistryModule::AssetCreated(NewBP);
 
-				// Mark the package dirty...
-				Package->MarkPackageDirty();
+					// Mark the package dirty...
+					Package->MarkPackageDirty();
+				}
 			}
 		}
 	}
