[BehaviorTreeEditor] fix for UBehaviorTreeDecoratorGraph not refreshing when changing properties in the details view

diff --git a/Engine/Source/Editor/BehaviorTreeEditor/Classes/EdGraphSchema_BehaviorTreeDecorator.h b/Engine/Source/Editor/BehaviorTreeEditor/Classes/EdGraphSchema_BehaviorTreeDecorator.h
--- a/Engine/Source/Editor/BehaviorTreeEditor/Classes/EdGraphSchema_BehaviorTreeDecorator.h
+++ b/Engine/Source/Editor/BehaviorTreeEditor/Classes/EdGraphSchema_BehaviorTreeDecorator.h
@@ -64,7 +64,13 @@ class UEdGraphSchema_BehaviorTreeDecorator : public UEdGraphSchema
 	virtual const FPinConnectionResponse CanCreateConnection(const UEdGraphPin* A, const UEdGraphPin* B) const override;
 	virtual FLinearColor GetPinTypeColor(const FEdGraphPinType& PinType) const override;
 	virtual bool ShouldHidePinDefaultValue(UEdGraphPin* Pin) const override;
+	virtual bool IsCacheVisualizationOutOfDate(int32 InVisualizationCacheID) const override;
+	virtual int32 GetCurrentVisualizationCacheID() const override;
+	virtual void ForceVisualizationCacheClear() const override;
 	//~ End EdGraphSchema Interface
 
 	static TSharedPtr<FDecoratorSchemaAction_NewNode> AddNewDecoratorAction(FGraphContextMenuBuilder& ContextMenuBuilder, const FText& Category, const FText& MenuDesc, const FText& Tooltip);
+
+private:
+	inline static int32 CurrentCacheRefreshID = 0;
 };
diff --git a/Engine/Source/Editor/BehaviorTreeEditor/Private/BehaviorTreeEditor.cpp b/Engine/Source/Editor/BehaviorTreeEditor/Private/BehaviorTreeEditor.cpp
--- a/Engine/Source/Editor/BehaviorTreeEditor/Private/BehaviorTreeEditor.cpp
+++ b/Engine/Source/Editor/BehaviorTreeEditor/Private/BehaviorTreeEditor.cpp
@@ -1106,7 +1106,12 @@ void FBehaviorTreeEditor::OnFinishedChangingProperties(const FPropertyChangedEve
 		MyGraph->UpdateInjectedNodes();
 		MyGraph->UpdateAsset(UBehaviorTreeGraph::ClearDebuggerFlags);
 	}
-	BehaviorTree->BTGraph->GetSchema()->ForceVisualizationCacheClear();
+
+	const TSharedPtr<SGraphEditor> FocusedGraphEd = UpdateGraphEdPtr.Pin();
+	if (FocusedGraphEd.IsValid() && FocusedGraphEd->GetCurrentGraph())
+	{
+		FocusedGraphEd->GetCurrentGraph()->GetSchema()->ForceVisualizationCacheClear();
+	}
 }
 
 void FBehaviorTreeEditor::OnPackageSaved(const FString& PackageFileName, UPackage* Package, FObjectPostSaveContext ObjectSaveContext)
diff --git a/Engine/Source/Editor/BehaviorTreeEditor/Private/EdGraphSchema_BehaviorTreeDecorator.cpp b/Engine/Source/Editor/BehaviorTreeEditor/Private/EdGraphSchema_BehaviorTreeDecorator.cpp
--- a/Engine/Source/Editor/BehaviorTreeEditor/Private/EdGraphSchema_BehaviorTreeDecorator.cpp
+++ b/Engine/Source/Editor/BehaviorTreeEditor/Private/EdGraphSchema_BehaviorTreeDecorator.cpp
@@ -216,4 +216,19 @@ bool UEdGraphSchema_BehaviorTreeDecorator::ShouldHidePinDefaultValue(UEdGraphPin
 	return true;
 }
 
+bool UEdGraphSchema_BehaviorTreeDecorator::IsCacheVisualizationOutOfDate(const int32 InVisualizationCacheID) const
+{
+	return CurrentCacheRefreshID != InVisualizationCacheID;
+}
+
+int32 UEdGraphSchema_BehaviorTreeDecorator::GetCurrentVisualizationCacheID() const
+{
+	return CurrentCacheRefreshID;
+}
+
+void UEdGraphSchema_BehaviorTreeDecorator::ForceVisualizationCacheClear() const
+{
+	++CurrentCacheRefreshID;
+}
+
 #undef LOCTEXT_NAMESPACE
