[BehaviorTree] disconnecting root node in behavior tree editor will now disable the tree.

diff --git a/Engine/Source/Editor/BehaviorTreeEditor/Private/BehaviorTreeGraph.cpp b/Engine/Source/Editor/BehaviorTreeEditor/Private/BehaviorTreeGraph.cpp
--- a/Engine/Source/Editor/BehaviorTreeEditor/Private/BehaviorTreeGraph.cpp
+++ b/Engine/Source/Editor/BehaviorTreeEditor/Private/BehaviorTreeGraph.cpp
@@ -164,18 +164,18 @@ void UBehaviorTreeGraph::UpdateAsset(int32 UpdateFlags)
 
 	// we can't look at pins until pin references have been fixed up post undo:
 	UEdGraphPin::ResolveAllPinReferences();
+
+	UBehaviorTreeGraphNode* Node = nullptr;
 	if (RootNode && RootNode->Pins.Num() > 0 && RootNode->Pins[0]->LinkedTo.Num() > 0)
 	{
-		UBehaviorTreeGraphNode* Node = Cast<UBehaviorTreeGraphNode>(RootNode->Pins[0]->LinkedTo[0]->GetOwningNode());
-		if (Node)
-		{
-			CreateBTFromGraph(Node);
+		Node = Cast<UBehaviorTreeGraphNode>(RootNode->Pins[0]->LinkedTo[0]->GetOwningNode());
+	}
 
-			if ((UpdateFlags & KeepRebuildCounter) == 0)
-			{
-				ModCounter++;
-			}
-		}
+	CreateBTFromGraph(Node);
+
+	if ((UpdateFlags & KeepRebuildCounter) == 0)
+	{
+		ModCounter++;
 	}
 
 	UpdateBlackboardChange();
@@ -897,44 +897,56 @@ namespace BTGraphHelpers
 void UBehaviorTreeGraph::CreateBTFromGraph(UBehaviorTreeGraphNode* RootEdNode)
 {
 	UBehaviorTree* BTAsset = Cast<UBehaviorTree>(GetOuter());
-	BTAsset->RootNode = NULL; //discard old tree
-
-	// let's create new tree from graph
-	uint16 ExecutionIndex = 0;
-	uint8 TreeDepth = 0;
-
-	BTAsset->RootNode = Cast<UBTCompositeNode>(RootEdNode->NodeInstance);
-	if (BTAsset->RootNode)
-	{
-		BTAsset->RootNode->InitializeNode(NULL, ExecutionIndex, 0, TreeDepth);
-		ExecutionIndex++;
-	}
 
-	// collect root level decorators
-	uint16 DummyIndex = MAX_uint16;
+	//discard old tree
+	BTAsset->RootNode = nullptr;
 	BTAsset->RootDecorators.Empty();
 	BTAsset->RootDecoratorOps.Empty();
-	BTGraphHelpers::CollectDecorators(BTAsset, RootEdNode, BTAsset->RootDecorators, BTAsset->RootDecoratorOps, false, NULL, &DummyIndex, 0, 0);
 
-	// connect tree nodes
-	BTGraphHelpers::CreateChildren(BTAsset, BTAsset->RootNode, RootEdNode, &ExecutionIndex, TreeDepth + 1); //-V595
+	if (RootEdNode != nullptr)
+	{
+		// let's create new tree from graph
+		uint16 ExecutionIndex = 0;
 
-	// mark root level nodes
-	BTGraphHelpers::ClearRootLevelFlags(this);
+		BTAsset->RootNode = Cast<UBTCompositeNode>(RootEdNode->NodeInstance);
+		if (BTAsset->RootNode)
+		{
+			BTAsset->RootNode->InitializeNode(/*ParentNode*/nullptr, ExecutionIndex, /*MemoryOffset*/0 , /*TreeDepth*/0);
+			ExecutionIndex++;
+		}
 
-	RootEdNode->bRootLevel = true;
-	for (int32 Index = 0; Index < RootEdNode->Decorators.Num(); Index++)
-	{
-		UBehaviorTreeGraphNode* Node = RootEdNode->Decorators[Index];
-		if (Node)
+		// collect root level decorators
+		uint16 DummyExecutionIndex = MAX_uint16;
+		BTGraphHelpers::CollectDecorators(BTAsset
+			, RootEdNode
+			, BTAsset->RootDecorators
+			, BTAsset->RootDecoratorOps
+			, /*bInitializeNodes*/false
+			, /*RootNode*/nullptr
+			, &DummyExecutionIndex
+			, /*TreeDepth*/0
+			, /*ChildIdx*/0);
+
+		// connect tree nodes
+		BTGraphHelpers::CreateChildren(BTAsset, BTAsset->RootNode, RootEdNode, &ExecutionIndex, /*TreeDepth*/1); //-V595
+
+		// mark root level nodes
+		BTGraphHelpers::ClearRootLevelFlags(this);
+
+		RootEdNode->bRootLevel = true;
+		for (int32 Index = 0; Index < RootEdNode->Decorators.Num(); Index++)
 		{
-			Node->bRootLevel = true;
+			UBehaviorTreeGraphNode* Node = RootEdNode->Decorators[Index];
+			if (Node)
+			{
+				Node->bRootLevel = true;
+			}
 		}
-	}
 
-	if (BTAsset->RootNode)
-	{
-		BTAsset->RootNode->InitializeComposite(ExecutionIndex - 1);
+		if (BTAsset->RootNode)
+		{
+			BTAsset->RootNode->InitializeComposite(ExecutionIndex - 1);
+		}
 	}
 
 	// Now remove any orphaned nodes left behind after regeneration
