[BehaviorTree] clear references to BT nodes now after pasting some graph nodes. These references will be set again when updating the BT from the graph nodes.
#rnx

diff --git a/Engine/Source/Editor/BehaviorTreeEditor/Classes/BehaviorTreeGraphNode_Composite.h b/Engine/Source/Editor/BehaviorTreeEditor/Classes/BehaviorTreeGraphNode_Composite.h
--- a/Engine/Source/Editor/BehaviorTreeEditor/Classes/BehaviorTreeGraphNode_Composite.h
+++ b/Engine/Source/Editor/BehaviorTreeEditor/Classes/BehaviorTreeGraphNode_Composite.h
@@ -28,4 +28,8 @@ class BEHAVIORTREEEDITOR_API UBehaviorTreeGraphNode_Composite : public UBehavior
 
 	/** check if node can accept breakpoints */
 	virtual bool CanPlaceBreakpoints() const override { return true; }
+
+protected:
+
+	virtual void PostPasteNode() override;
 };
diff --git a/Engine/Source/Editor/BehaviorTreeEditor/Classes/BehaviorTreeGraphNode_CompositeDecorator.h b/Engine/Source/Editor/BehaviorTreeEditor/Classes/BehaviorTreeGraphNode_CompositeDecorator.h
--- a/Engine/Source/Editor/BehaviorTreeEditor/Classes/BehaviorTreeGraphNode_CompositeDecorator.h
+++ b/Engine/Source/Editor/BehaviorTreeEditor/Classes/BehaviorTreeGraphNode_CompositeDecorator.h
@@ -55,6 +55,7 @@ class UBehaviorTreeGraphNode_CompositeDecorator : public UBehaviorTreeGraphNode
 
 	virtual void PrepareForCopying() override;
 	virtual void PostCopyNode() override;
+	virtual void PostPasteNode() override;
 
 	int32 SpawnMissingNodes(const TArray<class UBTDecorator*>& NodeInstances, const TArray<struct FBTDecoratorLogic>& Operations, int32 StartIndex);
 	void CollectDecoratorData(TArray<class UBTDecorator*>& NodeInstances, TArray<struct FBTDecoratorLogic>& Operations) const;
diff --git a/Engine/Source/Editor/BehaviorTreeEditor/Private/BehaviorTreeGraphNode_Composite.cpp b/Engine/Source/Editor/BehaviorTreeEditor/Private/BehaviorTreeGraphNode_Composite.cpp
--- a/Engine/Source/Editor/BehaviorTreeEditor/Private/BehaviorTreeGraphNode_Composite.cpp
+++ b/Engine/Source/Editor/BehaviorTreeEditor/Private/BehaviorTreeGraphNode_Composite.cpp
@@ -51,6 +51,19 @@ FText UBehaviorTreeGraphNode_Composite::GetTooltipText() const
 	return Super::GetTooltipText();
 }
 
+void UBehaviorTreeGraphNode_Composite::PostPasteNode()
+{
+	Super::PostPasteNode();
+	
+	// Clear our references to other nodes since this will be set when creating/updating the BT from the graph nodes
+	UBTCompositeNode* CompositeNode = Cast<UBTCompositeNode>(NodeInstance);
+	if (CompositeNode)
+	{
+		CompositeNode->Children.Reset();
+		CompositeNode->Services.Reset();
+	}
+}
+
 void UBehaviorTreeGraphNode_Composite::GetNodeContextMenuActions(UToolMenu* Menu, UGraphNodeContextMenuContext* Context) const
 {
 	AddContextMenuActionsDecorators(Menu, "BehaviorTreeGraphNode", Context);
diff --git a/Engine/Source/Editor/BehaviorTreeEditor/Private/BehaviorTreeGraphNode_CompositeDecorator.cpp b/Engine/Source/Editor/BehaviorTreeEditor/Private/BehaviorTreeGraphNode_CompositeDecorator.cpp
--- a/Engine/Source/Editor/BehaviorTreeEditor/Private/BehaviorTreeGraphNode_CompositeDecorator.cpp
+++ b/Engine/Source/Editor/BehaviorTreeEditor/Private/BehaviorTreeGraphNode_CompositeDecorator.cpp
@@ -112,6 +112,14 @@ void UBehaviorTreeGraphNode_CompositeDecorator::PostCopyNode()
 	}
 }
 
+void UBehaviorTreeGraphNode_CompositeDecorator::PostPasteNode()
+{
+	Super::PostPasteNode();
+
+	// Clear reference to the parent since it will be set when creating/updating the BT from the graph nodes
+	ParentNodeInstance = nullptr;
+}
+
 void UBehaviorTreeGraphNode_CompositeDecorator::ResetNodeOwner()
 {
 	Super::ResetNodeOwner();
