Fixed some unnecessary use of FDeprecateSlateVector2D

diff --git a/Engine/Source/Editor/BehaviorTreeEditor/Private/SGraphNode_BehaviorTree.h b/Engine/Source/Editor/BehaviorTreeEditor/Private/SGraphNode_BehaviorTree.h
--- a/Engine/Source/Editor/BehaviorTreeEditor/Private/SGraphNode_BehaviorTree.h
+++ b/Engine/Source/Editor/BehaviorTreeEditor/Private/SGraphNode_BehaviorTree.h
@@ -72,7 +72,7 @@ public:
 	/** shows red marker when search failed*/
 	EVisibility GetDebuggerSearchFailedMarkerVisibility() const;
 
-	UE::Slate::FDeprecateVector2DResult GetCachedPosition() const { return CachedPosition; }
+	FVector2f GetCachedPosition() const { return CachedPosition; }
 
 protected:
 	/** INotifyOnEnumChanged interface */
@@ -96,7 +96,7 @@ protected:
 	TArray<FNodeBounds> TriggerOffsets;
 
 	/** cached draw position */
-	FDeprecateSlateVector2D CachedPosition;
+	FVector2f CachedPosition;
 
 	TArray< TSharedPtr<SGraphNode> > DecoratorWidgets;
 	TArray< TSharedPtr<SGraphNode> > ServicesWidgets;
