Changed SBehaviorTreeBlackboardView to own a weakptr to its blackboard data instead of a strong ref

diff --git a/Engine/Source/Editor/BehaviorTreeEditor/Private/SBehaviorTreeBlackboardEditor.cpp b/Engine/Source/Editor/BehaviorTreeEditor/Private/SBehaviorTreeBlackboardEditor.cpp
--- a/Engine/Source/Editor/BehaviorTreeEditor/Private/SBehaviorTreeBlackboardEditor.cpp
+++ b/Engine/Source/Editor/BehaviorTreeEditor/Private/SBehaviorTreeBlackboardEditor.cpp
@@ -123,7 +123,8 @@ TSharedPtr<FExtender> SBehaviorTreeBlackboardEditor::GetToolbarExtender(TSharedR
 
 void SBehaviorTreeBlackboardEditor::HandleDeleteEntry()
 {
-	if (BlackboardData == nullptr)
+	UBlackboardData* BlackboardDataPtr = BlackboardData.Get();
+	if (!BlackboardDataPtr)
 	{
 		UE_LOG(LogBlackboardEditor, Error, TEXT("Trying to delete an entry from a blackboard while no Blackboard Asset is set!"));
 		return;
@@ -136,23 +137,23 @@ void SBehaviorTreeBlackboardEditor::HandleDeleteEntry()
 		if(BlackboardEntry != nullptr && !bIsInherited)
 		{
 			const FScopedTransaction Transaction(LOCTEXT("BlackboardEntryDeleteTransaction", "Delete Blackboard Entry"));
-			BlackboardData->SetFlags(RF_Transactional);
-			BlackboardData->Modify();
+			BlackboardDataPtr->SetFlags(RF_Transactional);
+			BlackboardDataPtr->Modify();
 
 			FProperty* KeysProperty = FindFProperty<FProperty>(UBlackboardData::StaticClass(), GET_MEMBER_NAME_CHECKED(UBlackboardData, Keys));
-			BlackboardData->PreEditChange(KeysProperty);
+			BlackboardDataPtr->PreEditChange(KeysProperty);
 		
-			for(int32 ItemIndex = 0; ItemIndex < BlackboardData->Keys.Num(); ItemIndex++)
+			for(int32 ItemIndex = 0; ItemIndex < BlackboardDataPtr->Keys.Num(); ItemIndex++)
 			{
-				if(BlackboardEntry == &BlackboardData->Keys[ItemIndex])
+				if(BlackboardEntry == &BlackboardDataPtr->Keys[ItemIndex])
 				{
-					BlackboardData->Keys.RemoveAt(ItemIndex);
+					BlackboardDataPtr->Keys.RemoveAt(ItemIndex);
 					break;
 				}
 			}
 
 			GraphActionMenu->RefreshAllActions(true);
-			OnBlackboardKeyChanged.ExecuteIfBound(BlackboardData, nullptr);
+			OnBlackboardKeyChanged.ExecuteIfBound(BlackboardDataPtr, nullptr);
 
 			// signal de-selection
 			if(OnEntrySelected.IsBound())
@@ -161,7 +162,7 @@ void SBehaviorTreeBlackboardEditor::HandleDeleteEntry()
 			}
 
 			FPropertyChangedEvent PropertyChangedEvent(KeysProperty, EPropertyChangeType::ArrayRemove);
-			BlackboardData->PostEditChangeProperty(PropertyChangedEvent);
+			BlackboardDataPtr->PostEditChangeProperty(PropertyChangedEvent);
 		}
 	}
 }
@@ -255,7 +256,8 @@ TSharedRef<SWidget> SBehaviorTreeBlackboardEditor::HandleCreateNewEntryMenu() co
 
 void SBehaviorTreeBlackboardEditor::HandleKeyClassPicked(UClass* InClass)
 {
-	if (BlackboardData == nullptr)
+	UBlackboardData* BlackboardDataPtr = BlackboardData.Get();
+	if (BlackboardDataPtr == nullptr)
 	{
 		UE_LOG(LogBlackboardEditor, Error, TEXT("Trying to delete an entry from a blackboard while no Blackboard Asset is set!"));
 		return;
@@ -267,11 +269,11 @@ void SBehaviorTreeBlackboardEditor::HandleKeyClassPicked(UClass* InClass)
 	check(InClass->IsChildOf(UBlackboardKeyType::StaticClass()));
 
 	const FScopedTransaction Transaction(LOCTEXT("BlackboardEntryAddTransaction", "Add Blackboard Entry"));
-	BlackboardData->SetFlags(RF_Transactional);
-	BlackboardData->Modify();
+	BlackboardDataPtr->SetFlags(RF_Transactional);
+	BlackboardDataPtr->Modify();
 
 	FProperty* KeysProperty = FindFProperty<FProperty>(UBlackboardData::StaticClass(), GET_MEMBER_NAME_CHECKED(UBlackboardData, Keys));
-	BlackboardData->PreEditChange(KeysProperty);
+	BlackboardDataPtr->PreEditChange(KeysProperty);
 
 	// create a name for this new key
 	FString NewKeyName = InClass->GetDisplayNameText().ToString();
@@ -296,8 +298,8 @@ void SBehaviorTreeBlackboardEditor::HandleKeyClassPicked(UClass* InClass)
 	};
 
 	// check for existing keys of the same name
-	for(const auto& Key : BlackboardData->Keys) { DuplicateFunction(Key); };
-	for(const auto& Key : BlackboardData->ParentKeys) { DuplicateFunction(Key); };
+	for(const auto& Key : BlackboardDataPtr->Keys) { DuplicateFunction(Key); };
+	for(const auto& Key : BlackboardDataPtr->ParentKeys) { DuplicateFunction(Key); };
 
 	if(IndexSuffix != -1)
 	{
@@ -306,12 +308,12 @@ void SBehaviorTreeBlackboardEditor::HandleKeyClassPicked(UClass* InClass)
 
 	FBlackboardEntry Entry;
 	Entry.EntryName = FName(*NewKeyName);
-	Entry.KeyType = NewObject<UBlackboardKeyType>(BlackboardData, InClass);
+	Entry.KeyType = NewObject<UBlackboardKeyType>(BlackboardDataPtr, InClass);
 
-	BlackboardData->Keys.Add(Entry);
+	BlackboardDataPtr->Keys.Add(Entry);
 
 	GraphActionMenu->RefreshAllActions(true);
-	OnBlackboardKeyChanged.ExecuteIfBound(BlackboardData, &BlackboardData->Keys.Last());
+	OnBlackboardKeyChanged.ExecuteIfBound(BlackboardDataPtr, &BlackboardDataPtr->Keys.Last());
 
 	GraphActionMenu->SelectItemByName(Entry.EntryName, ESelectInfo::OnMouseClick);
 
@@ -326,7 +328,7 @@ void SBehaviorTreeBlackboardEditor::HandleKeyClassPicked(UClass* InClass)
 	GraphActionMenu->OnRequestRenameOnActionNode();
 
 	FPropertyChangedEvent PropertyChangedEvent(KeysProperty, EPropertyChangeType::ArrayAdd);
-	BlackboardData->PostEditChangeProperty(PropertyChangedEvent);
+	BlackboardDataPtr->PostEditChangeProperty(PropertyChangedEvent);
 }
 
 bool SBehaviorTreeBlackboardEditor::CanCreateNewEntry() const
diff --git a/Engine/Source/Editor/BehaviorTreeEditor/Private/SBehaviorTreeBlackboardView.cpp b/Engine/Source/Editor/BehaviorTreeEditor/Private/SBehaviorTreeBlackboardView.cpp
--- a/Engine/Source/Editor/BehaviorTreeEditor/Private/SBehaviorTreeBlackboardView.cpp
+++ b/Engine/Source/Editor/BehaviorTreeEditor/Private/SBehaviorTreeBlackboardView.cpp
@@ -77,14 +77,14 @@ FName FEdGraphSchemaAction_BlackboardEntry::GetTypeId() const
 	return StaticGetTypeId(); 
 }
 
-FEdGraphSchemaAction_BlackboardEntry::FEdGraphSchemaAction_BlackboardEntry( UBlackboardData* InBlackboardData, FBlackboardEntry& InKey, bool bInIsInherited )
+FEdGraphSchemaAction_BlackboardEntry::FEdGraphSchemaAction_BlackboardEntry( TWeakObjectPtr<UBlackboardData> InBlackboardData, FBlackboardEntry& InKey, bool bInIsInherited )
 	: FEdGraphSchemaAction_Dummy()
 	, BlackboardData(InBlackboardData)
 	, Key(InKey)
 	, bIsInherited(bInIsInherited)
 	, bIsNew(false)
 {
-	check(BlackboardData);
+	check(BlackboardData.IsValid());
 	Update();
 }
 
@@ -228,41 +228,45 @@ private:
 		if(NewName != OldName)
 		{
 			TArray<UObject*> ExternalBTAssetsWithKeyReferences;
-			if(!BlackboardEntryAction->bIsNew && BlackboardEntryAction->BlackboardData)
+			UBlackboardData* BlackboardData = BlackboardEntryAction->BlackboardData.Get();
+			if (BlackboardData)
 			{
-				// Preload behavior trees before we transact otherwise they will add objects to 
-				// the transaction buffer whether we change them or not.
-				// Blueprint regeneration does this in UEdGraphNode::CreatePin.
-				LoadReferencerBehaviorTrees(*(BlackboardEntryAction->BlackboardData), ExternalBTAssetsWithKeyReferences);
-			}
+				if (!BlackboardEntryAction->bIsNew)
+				{
+					// Preload behavior trees before we transact otherwise they will add objects to 
+					// the transaction buffer whether we change them or not.
+					// Blueprint regeneration does this in UEdGraphNode::CreatePin.
+					LoadReferencerBehaviorTrees(*(BlackboardData), ExternalBTAssetsWithKeyReferences);
+				}
 
-			const FScopedTransaction Transaction(LOCTEXT("BlackboardEntryRenameTransaction", "Rename Blackboard Entry"));
-			BlackboardEntryAction->BlackboardData->SetFlags(RF_Transactional);
-			BlackboardEntryAction->BlackboardData->Modify();
-			BlackboardEntryAction->Key.EntryName = NewName;
+				const FScopedTransaction Transaction(LOCTEXT("BlackboardEntryRenameTransaction", "Rename Blackboard Entry"));
+				BlackboardData->SetFlags(RF_Transactional);
+				BlackboardData->Modify();
+				BlackboardEntryAction->Key.EntryName = NewName;
 
-			FProperty* KeysArrayProperty = FindFProperty<FProperty>(UBlackboardData::StaticClass(), GET_MEMBER_NAME_CHECKED(UBlackboardData, Keys));
-			FProperty* NameProperty = FindFProperty<FProperty>(FBlackboardEntry::StaticStruct(), GET_MEMBER_NAME_CHECKED(FBlackboardEntry, EntryName));
-			FEditPropertyChain PropertyChain;
-			PropertyChain.AddHead(KeysArrayProperty);
-			PropertyChain.AddTail(NameProperty);
-			PropertyChain.SetActiveMemberPropertyNode(KeysArrayProperty);
-			PropertyChain.SetActivePropertyNode(NameProperty);
+				FProperty* KeysArrayProperty = FindFProperty<FProperty>(UBlackboardData::StaticClass(), GET_MEMBER_NAME_CHECKED(UBlackboardData, Keys));
+				FProperty* NameProperty = FindFProperty<FProperty>(FBlackboardEntry::StaticStruct(), GET_MEMBER_NAME_CHECKED(FBlackboardEntry, EntryName));
+				FEditPropertyChain PropertyChain;
+				PropertyChain.AddHead(KeysArrayProperty);
+				PropertyChain.AddTail(NameProperty);
+				PropertyChain.SetActiveMemberPropertyNode(KeysArrayProperty);
+				PropertyChain.SetActivePropertyNode(NameProperty);
 
-			BlackboardEntryAction->BlackboardData->PreEditChange(PropertyChain);
+				BlackboardData->PreEditChange(PropertyChain);
 
-			BlackboardEntryAction->Update();
+				BlackboardEntryAction->Update();
 
-			OnBlackboardKeyChanged.ExecuteIfBound(BlackboardEntryAction->BlackboardData, &BlackboardEntryAction->Key);
+				OnBlackboardKeyChanged.ExecuteIfBound(BlackboardData, &BlackboardEntryAction->Key);
 
-			if(!BlackboardEntryAction->bIsNew)
-			{
-				UpdateExternalBlackboardKeyReferences(OldName, NewName, ExternalBTAssetsWithKeyReferences);
-			}
+				if (!BlackboardEntryAction->bIsNew)
+				{
+					UpdateExternalBlackboardKeyReferences(OldName, NewName, ExternalBTAssetsWithKeyReferences);
+				}
 
-			FPropertyChangedEvent PropertyChangedEvent(NameProperty, EPropertyChangeType::ValueSet);
-			FPropertyChangedChainEvent PropertyChangedChainEvent(PropertyChain, PropertyChangedEvent);
-			BlackboardEntryAction->BlackboardData->PostEditChangeChainProperty(PropertyChangedChainEvent);
+				FPropertyChangedEvent PropertyChangedEvent(NameProperty, EPropertyChangeType::ValueSet);
+				FPropertyChangedChainEvent PropertyChangedChainEvent(PropertyChain, PropertyChangedEvent);
+				BlackboardData->PostEditChangeChainProperty(PropertyChangedChainEvent);
+			}
 		}
 
 		BlackboardEntryAction->bIsNew = false;
@@ -355,21 +359,25 @@ private:
 		const FString NewTextAsString = InNewText.ToString();
 
 		// check for duplicate keys
-		for(const auto& Key : BlackboardEntryAction->BlackboardData->Keys)
+		const UBlackboardData* BlackboardData = BlackboardEntryAction->BlackboardData.Get();
+		if (BlackboardData)
 		{
-			if(&BlackboardEntryAction->Key != &Key && Key.EntryName.ToString() == NewTextAsString)
+			for(const auto& Key : BlackboardData->Keys)
 			{
-				OutErrorMessage = LOCTEXT("DuplicateKeyWarning", "A key of this name already exists.");
-				return false;
+				if(&BlackboardEntryAction->Key != &Key && Key.EntryName.ToString() == NewTextAsString)
+				{
+					OutErrorMessage = LOCTEXT("DuplicateKeyWarning", "A key of this name already exists.");
+					return false;
+				}
 			}
-		}
 
-		for(const auto& Key : BlackboardEntryAction->BlackboardData->ParentKeys)
-		{
-			if(&BlackboardEntryAction->Key != &Key && Key.EntryName.ToString() == NewTextAsString)
+			for(const auto& Key : BlackboardData->ParentKeys)
 			{
-				OutErrorMessage = LOCTEXT("DuplicateParentKeyWarning", "An inherited key of this name already exists.");
-				return false;
+				if(&BlackboardEntryAction->Key != &Key && Key.EntryName.ToString() == NewTextAsString)
+				{
+					OutErrorMessage = LOCTEXT("DuplicateParentKeyWarning", "An inherited key of this name already exists.");
+					return false;
+				}
 			}
 		}
 
@@ -436,14 +444,6 @@ private:
 };
 
 
-void SBehaviorTreeBlackboardView::AddReferencedObjects( FReferenceCollector& Collector )
-{
-	if(BlackboardData != nullptr)
-	{
-		Collector.AddReferencedObject(BlackboardData);
-	}
-}
-
 void SBehaviorTreeBlackboardView::Construct(const FArguments& InArgs, TSharedRef<FUICommandList> InCommandList, UBlackboardData* InBlackboardData)
 {
 	OnEntrySelected = InArgs._OnEntrySelected;
@@ -609,12 +609,16 @@ int32 SBehaviorTreeBlackboardView::GetSelectedEntryIndex(bool& bOutIsInherited)
 		FBlackboardEntry* BlackboardEntry = &Entry->Key;
 
 		// check to see what entry index we are using
-		TArray<FBlackboardEntry>& EntryArray = bOutIsInherited ? BlackboardData->ParentKeys : BlackboardData->Keys;
-		for(int32 Index = 0; Index < EntryArray.Num(); Index++)
+		const UBlackboardData* BlackboardDataPtr = BlackboardData.Get();
+		if (BlackboardDataPtr)
 		{
-			if(BlackboardEntry == &EntryArray[Index])
+			const TArray<FBlackboardEntry>& EntryArray = bOutIsInherited ? BlackboardDataPtr->ParentKeys : BlackboardDataPtr->Keys;
+			for (int32 Index = 0; Index < EntryArray.Num(); Index++)
 			{
-				return Index;
+				if (BlackboardEntry == &EntryArray[Index])
+				{
+					return Index;
+				}
 			}
 		}
 	}
diff --git a/Engine/Source/Editor/BehaviorTreeEditor/Private/SBehaviorTreeBlackboardView.h b/Engine/Source/Editor/BehaviorTreeEditor/Private/SBehaviorTreeBlackboardView.h
--- a/Engine/Source/Editor/BehaviorTreeEditor/Private/SBehaviorTreeBlackboardView.h
+++ b/Engine/Source/Editor/BehaviorTreeEditor/Private/SBehaviorTreeBlackboardView.h
@@ -54,12 +54,12 @@ public:
 	static FName StaticGetTypeId();
 	virtual FName GetTypeId() const;
 
-	FEdGraphSchemaAction_BlackboardEntry( UBlackboardData* InBlackboardData, FBlackboardEntry& InKey, bool bInIsInherited );
+	FEdGraphSchemaAction_BlackboardEntry( TWeakObjectPtr<UBlackboardData> InBlackboardData, FBlackboardEntry& InKey, bool bInIsInherited );
 
 	void Update();
 
 	/** Blackboard we reference our key in */
-	UBlackboardData* BlackboardData;
+	TWeakObjectPtr<UBlackboardData> BlackboardData;
 
 	/** Actual key */
 	FBlackboardEntry& Key;
@@ -72,7 +72,7 @@ public:
 };
 
 /** Displays blackboard entries */
-class SBehaviorTreeBlackboardView : public SCompoundWidget, public FGCObject
+class SBehaviorTreeBlackboardView : public SCompoundWidget
 {
 public:
 	SLATE_BEGIN_ARGS( SBehaviorTreeBlackboardView ) 
@@ -93,13 +93,6 @@ public:
 
 	void Construct(const FArguments& InArgs, TSharedRef<FUICommandList> InCommandList, UBlackboardData* InBlackboardData);
 
-	/** FGCObject implementation */
-	virtual void AddReferencedObjects( FReferenceCollector& Collector ) override;
-	virtual FString GetReferencerName() const override
-	{
-		return TEXT("SBehaviorTreeBlackboardView");
-	}
-
 	/**
 	 * Retrieves the blackboard item currently selected by the user.
 	 * @param	bOutIsInherited		Whether the item's index is in the key or parent key array
@@ -174,7 +167,7 @@ protected:
 
 protected:
 	/** The blackboard we are editing/viewing */
-	UBlackboardData* BlackboardData;
+	TWeakObjectPtr<UBlackboardData> BlackboardData;
 
 	/** The list of blackboard entries */
 	TSharedPtr<SGraphActionMenu> GraphActionMenu;
