[BehaviorTreeEditor] add FocusAttentionOnNode to IBehaviorTreeEditor and implement it in BehaviorTreeEditor

diff --git a/Engine/Source/Editor/BehaviorTreeEditor/Private/BehaviorTreeEditor.cpp b/Engine/Source/Editor/BehaviorTreeEditor/Private/BehaviorTreeEditor.cpp
--- a/Engine/Source/Editor/BehaviorTreeEditor/Private/BehaviorTreeEditor.cpp
+++ b/Engine/Source/Editor/BehaviorTreeEditor/Private/BehaviorTreeEditor.cpp
@@ -1602,7 +1602,13 @@ UEdGraphNode* FBehaviorTreeEditor::FindInjectedNode(int32 Index) const
 	return BTGraph ? BTGraph->FindInjectedNode(Index) : NULL;
 }
 
-void FBehaviorTreeEditor::DoubleClickNode(class UEdGraphNode* Node)
+void FBehaviorTreeEditor::DoubleClickNode(UEdGraphNode* Node)
+{
+	FocusAttentionOnNode(Node);
+	OnNodeDoubleClicked(Node);
+}
+
+void FBehaviorTreeEditor::FocusAttentionOnNode(UEdGraphNode* Node)
 {
 	TSharedPtr<SGraphEditor> CurrentGraphEditor = UpdateGraphEdPtr.Pin();
 	if (CurrentGraphEditor.IsValid())
@@ -1610,9 +1616,7 @@ void FBehaviorTreeEditor::DoubleClickNode(class UEdGraphNode* Node)
 		CurrentGraphEditor->ClearSelectionSet();
 		CurrentGraphEditor->SetNodeSelection(Node, true);
 	}
-
 	JumpToNode(Node);
-	OnNodeDoubleClicked(Node);
 }
 
 void FBehaviorTreeEditor::FocusWindow(UObject* ObjectToFocusOn)
diff --git a/Engine/Source/Editor/BehaviorTreeEditor/Private/BehaviorTreeEditor.h b/Engine/Source/Editor/BehaviorTreeEditor/Private/BehaviorTreeEditor.h
--- a/Engine/Source/Editor/BehaviorTreeEditor/Private/BehaviorTreeEditor.h
+++ b/Engine/Source/Editor/BehaviorTreeEditor/Private/BehaviorTreeEditor.h
@@ -60,6 +60,7 @@ public:
 	virtual void InitializeDebuggerState(class FBehaviorTreeDebugger* ParentDebugger) const override;
 	virtual UEdGraphNode* FindInjectedNode(int32 Index) const override;
 	virtual void DoubleClickNode(class UEdGraphNode* Node) override;
+	virtual void FocusAttentionOnNode(UEdGraphNode* Node) override;
 	virtual void FocusWindow(UObject* ObjectToFocusOn = NULL) override;
 	//~ End IBehaviorTreeEditor Interface
 
diff --git a/Engine/Source/Editor/BehaviorTreeEditor/Public/IBehaviorTreeEditor.h b/Engine/Source/Editor/BehaviorTreeEditor/Public/IBehaviorTreeEditor.h
--- a/Engine/Source/Editor/BehaviorTreeEditor/Public/IBehaviorTreeEditor.h
+++ b/Engine/Source/Editor/BehaviorTreeEditor/Public/IBehaviorTreeEditor.h
@@ -16,6 +16,7 @@ public:
 	virtual void InitializeDebuggerState(class FBehaviorTreeDebugger* ParentDebugger) const = 0;
 	virtual UEdGraphNode* FindInjectedNode(int32 Index) const = 0;
 	virtual void DoubleClickNode(class UEdGraphNode* Node) = 0;
+	virtual void FocusAttentionOnNode(UEdGraphNode* Node) = 0;
 };
 
 
