Expose some Behavior Tree Editor functionality to be inheritable and accessible from other modules.
Also introduce some smaller BT Editor UX improvements.

diff --git a/Engine/Source/Editor/AIGraph/Classes/AIGraphTypes.h b/Engine/Source/Editor/AIGraph/Classes/AIGraphTypes.h
--- a/Engine/Source/Editor/AIGraph/Classes/AIGraphTypes.h
+++ b/Engine/Source/Editor/AIGraph/Classes/AIGraphTypes.h
@@ -20,6 +20,7 @@ struct AIGRAPH_API FGraphNodeClassData
 	FString GetClassName() const;
 	FText GetCategory() const;
 	FString GetDisplayName() const;
+	FText GetTooltip() const;
 	UClass* GetClass(bool bSilent = false);
 	bool IsAbstract() const;
 
@@ -92,12 +93,22 @@ struct AIGRAPH_API FGraphNodeClassHelper
 	static void AddObservedBlueprintClasses(UClass* BaseNativeClass);
 	void UpdateAvailableBlueprintClasses();
 
+	/** Adds a single class to the list of hidden classes */
+	void AddForcedHiddenClass(UClass* Class);
+
+	/** Overrides all previously set hidden classes */
+	void SetForcedHiddenClasses(const TSet<UClass*>& Classes);
+
+	void SetGatherBlueprints(bool bGather);
+
 private:
 
 	UClass* RootNodeClass;
 	TSharedPtr<FGraphNodeClassNode> RootNode;
 	static TArray<FName> UnknownPackages;
 	static TMap<UClass*, int32> BlueprintClassCount;
+	TSet<UClass*> ForcedHiddenClasses;
+	bool bGatherBlueprints = true;
 
 	TSharedPtr<FGraphNodeClassNode> CreateClassDataNode(const struct FAssetData& AssetData);
 	TSharedPtr<FGraphNodeClassNode> FindBaseClassNode(TSharedPtr<FGraphNodeClassNode> Node, const FString& ClassName);
diff --git a/Engine/Source/Editor/AIGraph/Private/AIGraphTypes.cpp b/Engine/Source/Editor/AIGraph/Private/AIGraphTypes.cpp
--- a/Engine/Source/Editor/AIGraph/Private/AIGraphTypes.cpp
+++ b/Engine/Source/Editor/AIGraph/Private/AIGraphTypes.cpp
@@ -84,6 +84,11 @@ FString FGraphNodeClassData::GetDisplayName() const
 	return Class.IsValid() ? Class->GetMetaData(TEXT("DisplayName")) : FString();
 }
 
+FText FGraphNodeClassData::GetTooltip() const
+{
+	return Class.IsValid() ? Class->GetToolTipText() : FText::GetEmpty();
+}
+
 FText FGraphNodeClassData::GetCategory() const
 {
 	return Class.IsValid() ? FObjectEditorUtils::GetCategoryText(Class.Get()) : Category;
@@ -249,7 +254,11 @@ bool FGraphNodeClassHelper::IsHidingParentClass(UClass* Class)
 bool FGraphNodeClassHelper::IsHidingClass(UClass* Class)
 {
 	static FName MetaHideInEditor = TEXT("HiddenNode");
-	return Class && Class->HasAnyClassFlags(CLASS_Native) && Class->HasMetaData(MetaHideInEditor);
+
+	return 
+		Class && 
+		((Class->HasAnyClassFlags(CLASS_Native) && Class->HasMetaData(MetaHideInEditor))
+		|| ForcedHiddenClasses.Contains(Class));
 }
 
 bool FGraphNodeClassHelper::IsPackageSaved(FName PackageName)
@@ -461,17 +470,20 @@ void FGraphNodeClassHelper::BuildClassGraph()
 	}
 
 	// gather all blueprints
-	FAssetRegistryModule& AssetRegistryModule = FModuleManager::LoadModuleChecked<FAssetRegistryModule>(TEXT("AssetRegistry"));
-	TArray<FAssetData> BlueprintList;
+	if (bGatherBlueprints)
+	{
+		FAssetRegistryModule& AssetRegistryModule = FModuleManager::LoadModuleChecked<FAssetRegistryModule>(TEXT("AssetRegistry"));
+		TArray<FAssetData> BlueprintList;
 
-	FARFilter Filter;
-	Filter.ClassPaths.Add(UBlueprint::StaticClass()->GetClassPathName());
-	AssetRegistryModule.Get().GetAssets(Filter, BlueprintList);
+		FARFilter Filter;
+		Filter.ClassPaths.Add(UBlueprint::StaticClass()->GetClassPathName());
+		AssetRegistryModule.Get().GetAssets(Filter, BlueprintList);
 
-	for (int32 i = 0; i < BlueprintList.Num(); i++)
-	{
-		TSharedPtr<FGraphNodeClassNode> NewNode = CreateClassDataNode(BlueprintList[i]);
-		NodeList.Add(NewNode);
+		for (int32 i = 0; i < BlueprintList.Num(); i++)
+		{
+			TSharedPtr<FGraphNodeClassNode> NewNode = CreateClassDataNode(BlueprintList[i]);
+			NodeList.Add(NewNode);
+		}
 	}
 
 	// build class tree
@@ -535,4 +547,22 @@ void FGraphNodeClassHelper::UpdateAvailableBlueprintClasses()
 	}
 }
 
+void FGraphNodeClassHelper::AddForcedHiddenClass(UClass* Class)
+{
+	if (Class)
+	{
+		ForcedHiddenClasses.Add(Class);
+	}
+}
+
+void FGraphNodeClassHelper::SetForcedHiddenClasses(const TSet<UClass*>& Classes)
+{
+	ForcedHiddenClasses = Classes;
+}
+
+void FGraphNodeClassHelper::SetGatherBlueprints(const bool bGather)
+{
+	bGatherBlueprints = bGather;
+}
+
 #undef LOCTEXT_NAMESPACE
diff --git a/Engine/Source/Editor/BehaviorTreeEditor/Classes/BehaviorTreeDecoratorGraph.h b/Engine/Source/Editor/BehaviorTreeEditor/Classes/BehaviorTreeDecoratorGraph.h
--- a/Engine/Source/Editor/BehaviorTreeEditor/Classes/BehaviorTreeDecoratorGraph.h
+++ b/Engine/Source/Editor/BehaviorTreeEditor/Classes/BehaviorTreeDecoratorGraph.h
@@ -17,7 +17,7 @@ class UEdGraphPin;
 class UObject;
 
 UCLASS()
-class UBehaviorTreeDecoratorGraph : public UEdGraph
+class BEHAVIORTREEEDITOR_API UBehaviorTreeDecoratorGraph : public UEdGraph
 {
 	GENERATED_UCLASS_BODY()
 
diff --git a/Engine/Source/Editor/BehaviorTreeEditor/Classes/BehaviorTreeFactory.h b/Engine/Source/Editor/BehaviorTreeEditor/Classes/BehaviorTreeFactory.h
--- a/Engine/Source/Editor/BehaviorTreeEditor/Classes/BehaviorTreeFactory.h
+++ b/Engine/Source/Editor/BehaviorTreeEditor/Classes/BehaviorTreeFactory.h
@@ -15,7 +15,7 @@ class UClass;
 class UObject;
 
 UCLASS()
-class UBehaviorTreeFactory : public UFactory
+class BEHAVIORTREEEDITOR_API UBehaviorTreeFactory : public UFactory
 {
 	GENERATED_UCLASS_BODY()
 
diff --git a/Engine/Source/Editor/BehaviorTreeEditor/Classes/BehaviorTreeGraph.h b/Engine/Source/Editor/BehaviorTreeEditor/Classes/BehaviorTreeGraph.h
--- a/Engine/Source/Editor/BehaviorTreeEditor/Classes/BehaviorTreeGraph.h
+++ b/Engine/Source/Editor/BehaviorTreeEditor/Classes/BehaviorTreeGraph.h
@@ -11,11 +11,12 @@
 
 #include "BehaviorTreeGraph.generated.h"
 
+class UBehaviorTreeGraphNode_Root;
 class UEdGraphNode;
 class UObject;
 
 UCLASS()
-class UBehaviorTreeGraph : public UAIGraph
+class BEHAVIORTREEEDITOR_API UBehaviorTreeGraph : public UAIGraph
 {
 	GENERATED_UCLASS_BODY()
 
@@ -33,6 +34,9 @@ class UBehaviorTreeGraph : public UAIGraph
 	UPROPERTY()
 	bool bIsUsingModCounter;
 
+	UPROPERTY()
+	TSubclassOf<UBehaviorTreeGraphNode_Root> RootNodeClass;
+
 	virtual void OnCreated() override;
 	virtual void OnLoaded() override;
 	virtual void Initialize() override;
diff --git a/Engine/Source/Editor/BehaviorTreeEditor/Classes/BehaviorTreeGraphNode.h b/Engine/Source/Editor/BehaviorTreeEditor/Classes/BehaviorTreeGraphNode.h
--- a/Engine/Source/Editor/BehaviorTreeEditor/Classes/BehaviorTreeGraphNode.h
+++ b/Engine/Source/Editor/BehaviorTreeEditor/Classes/BehaviorTreeGraphNode.h
@@ -14,6 +14,7 @@
 
 #include "BehaviorTreeGraphNode.generated.h"
 
+class ISlateStyle;
 class UEdGraph;
 class UEdGraphSchema;
 class UObject;
diff --git a/Engine/Source/Editor/BehaviorTreeEditor/Classes/BehaviorTreeGraphNode_CompositeDecorator.h b/Engine/Source/Editor/BehaviorTreeEditor/Classes/BehaviorTreeGraphNode_CompositeDecorator.h
--- a/Engine/Source/Editor/BehaviorTreeEditor/Classes/BehaviorTreeGraphNode_CompositeDecorator.h
+++ b/Engine/Source/Editor/BehaviorTreeEditor/Classes/BehaviorTreeGraphNode_CompositeDecorator.h
@@ -16,11 +16,12 @@
 
 #include "BehaviorTreeGraphNode_CompositeDecorator.generated.h"
 
+class UBehaviorTreeDecoratorGraph;
 class UEdGraph;
 class UObject;
 
 UCLASS()
-class UBehaviorTreeGraphNode_CompositeDecorator : public UBehaviorTreeGraphNode
+class BEHAVIORTREEEDITOR_API UBehaviorTreeGraphNode_CompositeDecorator : public UBehaviorTreeGraphNode
 {
 	GENERATED_UCLASS_BODY()
 
@@ -41,6 +42,8 @@ class UBehaviorTreeGraphNode_CompositeDecorator : public UBehaviorTreeGraphNode
 
 	uint32 bHasBrokenInstances : 1;
 
+	TSubclassOf<UBehaviorTreeDecoratorGraph> GraphClass;
+
 	FString GetNodeTypeDescription() const;
 	virtual FText GetNodeTitle(ENodeTitleType::Type TitleType) const override;
 	virtual void AllocateDefaultPins() override;
diff --git a/Engine/Source/Editor/BehaviorTreeEditor/Classes/BehaviorTreeGraphNode_Root.h b/Engine/Source/Editor/BehaviorTreeEditor/Classes/BehaviorTreeGraphNode_Root.h
--- a/Engine/Source/Editor/BehaviorTreeEditor/Classes/BehaviorTreeGraphNode_Root.h
+++ b/Engine/Source/Editor/BehaviorTreeEditor/Classes/BehaviorTreeGraphNode_Root.h
@@ -17,7 +17,7 @@ class UObject;
 
 /** Root node of this behavior tree, holds Blackboard data */
 UCLASS()
-class UBehaviorTreeGraphNode_Root : public UBehaviorTreeGraphNode
+class BEHAVIORTREEEDITOR_API UBehaviorTreeGraphNode_Root : public UBehaviorTreeGraphNode
 {
 	GENERATED_UCLASS_BODY()
 
diff --git a/Engine/Source/Editor/BehaviorTreeEditor/Classes/EdGraphSchema_BehaviorTree.h b/Engine/Source/Editor/BehaviorTreeEditor/Classes/EdGraphSchema_BehaviorTree.h
--- a/Engine/Source/Editor/BehaviorTreeEditor/Classes/EdGraphSchema_BehaviorTree.h
+++ b/Engine/Source/Editor/BehaviorTreeEditor/Classes/EdGraphSchema_BehaviorTree.h
@@ -4,6 +4,7 @@
 
 #include "AIGraphSchema.h"
 #include "AIGraphTypes.h"
+#include "BehaviorTreeGraphNode_CompositeDecorator.h"
 #include "Containers/Array.h"
 #include "CoreMinimal.h"
 #include "EdGraph/EdGraphSchema.h"
@@ -18,6 +19,7 @@
 #include "EdGraphSchema_BehaviorTree.generated.h"
 
 class FSlateRect;
+class UBehaviorTreeGraphNode_CompositeDecorator;
 class UClass;
 class UEdGraph;
 class UEdGraphNode;
@@ -44,8 +46,8 @@ struct FBehaviorTreeSchemaAction_AutoArrange : public FEdGraphSchemaAction
 	//~ End FEdGraphSchemaAction Interface
 };
 
-UCLASS(MinimalAPI)
-class UEdGraphSchema_BehaviorTree : public UAIGraphSchema
+UCLASS()
+class BEHAVIORTREEEDITOR_API UEdGraphSchema_BehaviorTree : public UAIGraphSchema
 {
 	GENERATED_UCLASS_BODY()
 
@@ -65,6 +67,12 @@ class UEdGraphSchema_BehaviorTree : public UAIGraphSchema
 	virtual void GetGraphNodeContextActions(FGraphContextMenuBuilder& ContextMenuBuilder, int32 SubNodeFlags) const override;
 	virtual void GetSubNodeClasses(int32 SubNodeFlags, TArray<FGraphNodeClassData>& ClassData, UClass*& GraphNodeClass) const override;
 
+	TSubclassOf<UBehaviorTreeGraphNode_CompositeDecorator> CompositeDecoratorClass;
+
+protected:
+	virtual FGraphNodeClassHelper& GetClassCache() const;
+	virtual bool IsNodeSubtreeTask(const FGraphNodeClassData& NodeClass) const;
+
 private:
 	// ID for checking dirty status of node titles against, increases whenever 
 	static int32 CurrentCacheRefreshID;
diff --git a/Engine/Source/Editor/BehaviorTreeEditor/Classes/EdGraphSchema_BehaviorTreeDecorator.h b/Engine/Source/Editor/BehaviorTreeEditor/Classes/EdGraphSchema_BehaviorTreeDecorator.h
--- a/Engine/Source/Editor/BehaviorTreeEditor/Classes/EdGraphSchema_BehaviorTreeDecorator.h
+++ b/Engine/Source/Editor/BehaviorTreeEditor/Classes/EdGraphSchema_BehaviorTreeDecorator.h
@@ -25,6 +25,7 @@ class UEdGraphNode;
 class UEdGraphPin;
 class UObject;
 struct FEdGraphPinType;
+struct FGraphNodeClassHelper;
 
 /** Action to add a node to the graph */
 USTRUCT()
@@ -62,8 +63,8 @@ struct FDecoratorSchemaAction_NewNode : public FEdGraphSchemaAction
 	}
 };
 
-UCLASS(MinimalAPI)
-class UEdGraphSchema_BehaviorTreeDecorator : public UEdGraphSchema
+UCLASS()
+class BEHAVIORTREEEDITOR_API UEdGraphSchema_BehaviorTreeDecorator : public UEdGraphSchema
 {
 	GENERATED_UCLASS_BODY()
 
@@ -88,6 +89,9 @@ class UEdGraphSchema_BehaviorTreeDecorator : public UEdGraphSchema
 
 	static TSharedPtr<FDecoratorSchemaAction_NewNode> AddNewDecoratorAction(FGraphContextMenuBuilder& ContextMenuBuilder, const FText& Category, const FText& MenuDesc, const FText& Tooltip);
 
+protected:
+	virtual FGraphNodeClassHelper& GetClassCache() const;
+
 private:
 	inline static int32 CurrentCacheRefreshID = 0;
 };
diff --git a/Engine/Source/Editor/BehaviorTreeEditor/Private/AssetDefinition_BehaviorTree.h b/Engine/Source/Editor/BehaviorTreeEditor/Private/AssetDefinition_BehaviorTree.h
deleted file mode 100644
diff --git a/Engine/Source/Editor/BehaviorTreeEditor/Private/BehaviorTreeEditor.cpp b/Engine/Source/Editor/BehaviorTreeEditor/Private/BehaviorTreeEditor.cpp
--- a/Engine/Source/Editor/BehaviorTreeEditor/Private/BehaviorTreeEditor.cpp
+++ b/Engine/Source/Editor/BehaviorTreeEditor/Private/BehaviorTreeEditor.cpp
@@ -139,6 +139,11 @@ FBehaviorTreeEditor::FBehaviorTreeEditor()
 	BlackboardData = nullptr;
 
 	bCheckDirtyOnAssetSave = true;
+
+	GraphClass = UBehaviorTreeGraph::StaticClass();
+	GraphName = "Behavior Tree";
+	CornerText = LOCTEXT("AppearanceCornerText", "BEHAVIOR TREE");
+	TitleText = LOCTEXT("BehaviorTreeGraphLabel", "Behavior Tree");
 }
 
 FBehaviorTreeEditor::~FBehaviorTreeEditor()
@@ -337,7 +342,9 @@ void FBehaviorTreeEditor::RestoreBehaviorTree()
 	const bool bNewGraph = MyGraph == NULL;
 	if (MyGraph == NULL)
 	{
-		BehaviorTree->BTGraph = FBlueprintEditorUtils::CreateNewGraph(BehaviorTree, TEXT("Behavior Tree"), UBehaviorTreeGraph::StaticClass(), UEdGraphSchema_BehaviorTree::StaticClass());
+		const TSubclassOf<UEdGraphSchema> SchemaClass = GetDefault<UBehaviorTreeGraph>(GraphClass)->Schema;
+		check(SchemaClass);
+		BehaviorTree->BTGraph = FBlueprintEditorUtils::CreateNewGraph(BehaviorTree, GraphName, GraphClass, SchemaClass);
 		MyGraph = Cast<UBehaviorTreeGraph>(BehaviorTree->BTGraph);
 
 		// Initialize the behavior tree graph
@@ -516,7 +523,7 @@ EVisibility FBehaviorTreeEditor::GetRootLevelNodeVisibility() const
 FGraphAppearanceInfo FBehaviorTreeEditor::GetGraphAppearance() const
 {
 	FGraphAppearanceInfo AppearanceInfo;
-	AppearanceInfo.CornerText = LOCTEXT("AppearanceCornerText", "BEHAVIOR TREE");
+	AppearanceInfo.CornerText = CornerText;
 
 	const int32 StepIdx = Debugger.IsValid() ? Debugger->GetShownStateIndex() : 0;
 	if (Debugger.IsValid() && !Debugger->IsDebuggerRunning())
@@ -630,7 +637,7 @@ TSharedRef<SGraphEditor> FBehaviorTreeEditor::CreateGraphEditorWidget(UEdGraph*
 			.FillWidth(1.f)
 			[
 				SNew(STextBlock)
-				.Text(LOCTEXT("BehaviorTreeGraphLabel", "Behavior Tree"))
+				.Text(TitleText)
 				.TextStyle( FAppStyle::Get(), TEXT("GraphBreadcrumbButtonText") )
 			]
 		];
@@ -643,7 +650,8 @@ TSharedRef<SGraphEditor> FBehaviorTreeEditor::CreateGraphEditorWidget(UEdGraph*
 		.Appearance(this, &FBehaviorTreeEditor::GetGraphAppearance)
 		.TitleBar(TitleBarWidget)
 		.GraphToEdit(InGraph)
-		.GraphEvents(InEvents);
+		.GraphEvents(InEvents)
+		.AutoExpandActionMenu(true);
 }
 
 bool FBehaviorTreeEditor::InEditingMode(bool bGraphIsEditable) const
@@ -1395,6 +1403,11 @@ void FBehaviorTreeEditor::SaveAsset_Execute()
 	IBehaviorTreeEditor::SaveAsset_Execute();
 }
 
+void FBehaviorTreeEditor::SetToolbarCreateActionsEnabled(const bool bActionsEnabled)
+{
+	ToolbarBuilder->SetCreateActionsEnabled(bActionsEnabled);
+}
+
 void FBehaviorTreeEditor::OnEnableBreakpoint()
 {
 	const FGraphPanelSelectionSet SelectedNodes = GetSelectedNodes();
diff --git a/Engine/Source/Editor/BehaviorTreeEditor/Private/BehaviorTreeEditor.h b/Engine/Source/Editor/BehaviorTreeEditor/Private/BehaviorTreeEditor.h
deleted file mode 100644
diff --git a/Engine/Source/Editor/BehaviorTreeEditor/Private/BehaviorTreeEditorToolbar.cpp b/Engine/Source/Editor/BehaviorTreeEditor/Private/BehaviorTreeEditorToolbar.cpp
--- a/Engine/Source/Editor/BehaviorTreeEditor/Private/BehaviorTreeEditorToolbar.cpp
+++ b/Engine/Source/Editor/BehaviorTreeEditor/Private/BehaviorTreeEditorToolbar.cpp
@@ -128,6 +128,11 @@ void FBehaviorTreeEditorToolbar::AddBehaviorTreeToolbar(TSharedPtr<FExtender> Ex
 	BehaviorTreeEditorPtr->AddToolbarExtender(ToolbarExtender);
 }
 
+void FBehaviorTreeEditorToolbar::SetCreateActionsEnabled(const bool bActionsEnabled)
+{
+	bCreateActionsEnabled = bActionsEnabled;
+}
+
 void FBehaviorTreeEditorToolbar::FillModesToolbar(FToolBarBuilder& ToolbarBuilder)
 {
 	check(BehaviorTreeEditor.IsValid());
@@ -165,6 +170,11 @@ void FBehaviorTreeEditorToolbar::FillModesToolbar(FToolBarBuilder& ToolbarBuilde
 
 void FBehaviorTreeEditorToolbar::FillBehaviorTreeToolbar(FToolBarBuilder& ToolbarBuilder)
 {
+	if (!bCreateActionsEnabled)
+	{
+		return;
+	}
+
 	check(BehaviorTreeEditor.IsValid());
 	TSharedPtr<FBehaviorTreeEditor> BehaviorTreeEditorPtr = BehaviorTreeEditor.Pin();
 
diff --git a/Engine/Source/Editor/BehaviorTreeEditor/Private/BehaviorTreeEditorToolbar.h b/Engine/Source/Editor/BehaviorTreeEditor/Private/BehaviorTreeEditorToolbar.h
--- a/Engine/Source/Editor/BehaviorTreeEditor/Private/BehaviorTreeEditorToolbar.h
+++ b/Engine/Source/Editor/BehaviorTreeEditor/Private/BehaviorTreeEditorToolbar.h
@@ -18,6 +18,8 @@ public:
 	void AddDebuggerToolbar(TSharedPtr<FExtender> Extender);
 	void AddBehaviorTreeToolbar(TSharedPtr<FExtender> Extender);
 
+	void SetCreateActionsEnabled(bool bActionsEnabled);
+
 private:
 	void FillModesToolbar(FToolBarBuilder& ToolbarBuilder);
 	void FillDebuggerToolbar(FToolBarBuilder& ToolbarBuilder);
@@ -26,4 +28,7 @@ private:
 protected:
 	/** Pointer back to the blueprint editor tool that owns us */
 	TWeakPtr<FBehaviorTreeEditor> BehaviorTreeEditor;
+
+private:
+	bool bCreateActionsEnabled = true;
 };
diff --git a/Engine/Source/Editor/BehaviorTreeEditor/Private/BehaviorTreeGraph.cpp b/Engine/Source/Editor/BehaviorTreeEditor/Private/BehaviorTreeGraph.cpp
--- a/Engine/Source/Editor/BehaviorTreeEditor/Private/BehaviorTreeGraph.cpp
+++ b/Engine/Source/Editor/BehaviorTreeEditor/Private/BehaviorTreeGraph.cpp
@@ -42,6 +42,7 @@ namespace BTGraphVersion
 UBehaviorTreeGraph::UBehaviorTreeGraph(const FObjectInitializer& ObjectInitializer) : Super(ObjectInitializer)
 {
 	Schema = UEdGraphSchema_BehaviorTree::StaticClass();
+	RootNodeClass = UBehaviorTreeGraphNode_Root::StaticClass();
 }
 
 void UBehaviorTreeGraph::UpdateBlackboardChange()
diff --git a/Engine/Source/Editor/BehaviorTreeEditor/Private/BehaviorTreeGraphNode_CompositeDecorator.cpp b/Engine/Source/Editor/BehaviorTreeEditor/Private/BehaviorTreeGraphNode_CompositeDecorator.cpp
--- a/Engine/Source/Editor/BehaviorTreeEditor/Private/BehaviorTreeGraphNode_CompositeDecorator.cpp
+++ b/Engine/Source/Editor/BehaviorTreeEditor/Private/BehaviorTreeGraphNode_CompositeDecorator.cpp
@@ -33,6 +33,8 @@ UBehaviorTreeGraphNode_CompositeDecorator::UBehaviorTreeGraphNode_CompositeDecor
 
 	FirstExecutionIndex = INDEX_NONE;
 	LastExecutionIndex = INDEX_NONE;
+
+	GraphClass = UBehaviorTreeDecoratorGraph::StaticClass();
 }
 
 void UBehaviorTreeGraphNode_CompositeDecorator::ResetExecutionRange()
@@ -199,8 +201,11 @@ void UBehaviorTreeGraphNode_CompositeDecorator::CreateBoundGraph()
 	// Create a new animation graph
 	check(BoundGraph == NULL);
 
+	const TSubclassOf<UEdGraphSchema> SchemaClass = GetDefault<UBehaviorTreeDecoratorGraph>(GraphClass)->Schema;
+	check(SchemaClass);
+
 	// don't use white space in name here, it prevents links from being copied correctly
-	BoundGraph = FBlueprintEditorUtils::CreateNewGraph(this, TEXT("CompositeDecorator"), UBehaviorTreeDecoratorGraph::StaticClass(), UEdGraphSchema_BehaviorTreeDecorator::StaticClass());
+	BoundGraph = FBlueprintEditorUtils::CreateNewGraph(this, TEXT("CompositeDecorator"), GraphClass, SchemaClass);
 	check(BoundGraph);
 
 	// Initialize the anim graph
diff --git a/Engine/Source/Editor/BehaviorTreeEditor/Private/EdGraphSchema_BehaviorTree.cpp b/Engine/Source/Editor/BehaviorTreeEditor/Private/EdGraphSchema_BehaviorTree.cpp
--- a/Engine/Source/Editor/BehaviorTreeEditor/Private/EdGraphSchema_BehaviorTree.cpp
+++ b/Engine/Source/Editor/BehaviorTreeEditor/Private/EdGraphSchema_BehaviorTree.cpp
@@ -66,12 +66,13 @@ UEdGraphNode* FBehaviorTreeSchemaAction_AutoArrange::PerformAction(class UEdGrap
 
 UEdGraphSchema_BehaviorTree::UEdGraphSchema_BehaviorTree(const FObjectInitializer& ObjectInitializer) : Super(ObjectInitializer)
 {
+	CompositeDecoratorClass = UBehaviorTreeGraphNode_CompositeDecorator::StaticClass();
 }
 
 void UEdGraphSchema_BehaviorTree::CreateDefaultNodesForGraph(UEdGraph& Graph) const
 {
 	FGraphNodeCreator<UBehaviorTreeGraphNode_Root> NodeCreator(Graph);
-	UBehaviorTreeGraphNode_Root* MyNode = NodeCreator.CreateNode();
+	UBehaviorTreeGraphNode_Root* MyNode = NodeCreator.CreateNode(true, Cast<UBehaviorTreeGraph>(&Graph)->RootNodeClass);
 	NodeCreator.Finalize();
 	SetNodeMetaData(MyNode, FNodeMetadata::DefaultGraphNode);
 }
@@ -82,10 +83,10 @@ void UEdGraphSchema_BehaviorTree::GetGraphNodeContextActions(FGraphContextMenuBu
 
 	if (SubNodeFlags == ESubNode::Decorator)
 	{
-		const FText Category = FObjectEditorUtils::GetCategoryText(UBehaviorTreeGraphNode_CompositeDecorator::StaticClass());
+		const FText Category = FObjectEditorUtils::GetCategoryText(CompositeDecoratorClass);
 		UEdGraph* Graph = (UEdGraph*)ContextMenuBuilder.CurrentGraph;
-		UBehaviorTreeGraphNode_CompositeDecorator* OpNode = NewObject<UBehaviorTreeGraphNode_CompositeDecorator>(Graph);
-		TSharedPtr<FAISchemaAction_NewSubNode> AddOpAction = UAIGraphSchema::AddNewSubNodeAction(ContextMenuBuilder, Category, FText::FromString(OpNode->GetNodeTypeDescription()), FText::GetEmpty());
+		UBehaviorTreeGraphNode_CompositeDecorator* OpNode = NewObject<UBehaviorTreeGraphNode_CompositeDecorator>(Graph, CompositeDecoratorClass);
+		TSharedPtr<FAISchemaAction_NewSubNode> AddOpAction = UAIGraphSchema::AddNewSubNodeAction(ContextMenuBuilder, Category, FText::FromString(OpNode->GetNodeTypeDescription()), OpNode->GetTooltipText());
 		AddOpAction->ParentNode = Cast<UBehaviorTreeGraphNode>(ContextMenuBuilder.SelectedObjects[0]);
 		AddOpAction->NodeTemplate = OpNode;
 	}
@@ -93,21 +94,33 @@ void UEdGraphSchema_BehaviorTree::GetGraphNodeContextActions(FGraphContextMenuBu
 
 void UEdGraphSchema_BehaviorTree::GetSubNodeClasses(int32 SubNodeFlags, TArray<FGraphNodeClassData>& ClassData, UClass*& GraphNodeClass) const
 {
-	FBehaviorTreeEditorModule& EditorModule = FModuleManager::GetModuleChecked<FBehaviorTreeEditorModule>(TEXT("BehaviorTreeEditor"));
-	FGraphNodeClassHelper* ClassCache = EditorModule.GetClassCache().Get();
+	FGraphNodeClassHelper& ClassCache = GetClassCache();
 
 	if (SubNodeFlags == ESubNode::Decorator)
 	{
-		ClassCache->GatherClasses(UBTDecorator::StaticClass(), ClassData);
+		ClassCache.GatherClasses(UBTDecorator::StaticClass(), ClassData);
 		GraphNodeClass = UBehaviorTreeGraphNode_Decorator::StaticClass();
 	}
 	else
 	{
-		ClassCache->GatherClasses(UBTService::StaticClass(), ClassData);
+		ClassCache.GatherClasses(UBTService::StaticClass(), ClassData);
 		GraphNodeClass = UBehaviorTreeGraphNode_Service::StaticClass();
 	}
 }
 
+FGraphNodeClassHelper& UEdGraphSchema_BehaviorTree::GetClassCache() const
+{
+	const FBehaviorTreeEditorModule& EditorModule = FModuleManager::GetModuleChecked<FBehaviorTreeEditorModule>(TEXT("BehaviorTreeEditor"));
+	FGraphNodeClassHelper* ClassHelper = EditorModule.GetClassCache().Get();
+	check(ClassHelper);
+	return *ClassHelper;
+}
+
+bool UEdGraphSchema_BehaviorTree::IsNodeSubtreeTask(const FGraphNodeClassData& NodeClass) const
+{
+	return NodeClass.GetClassName() == UBTTask_RunBehavior::StaticClass()->GetName();
+}
+
 void UEdGraphSchema_BehaviorTree::GetGraphContextActions(FGraphContextMenuBuilder& ContextMenuBuilder) const
 {
 	const FName PinCategory = ContextMenuBuilder.FromPin ?
@@ -120,22 +133,21 @@ void UEdGraphSchema_BehaviorTree::GetGraphContextActions(FGraphContextMenuBuilde
 	const bool bAllowComposites = bNoParent || !bOnlyTasks || bOnlyComposites;
 	const bool bAllowTasks = bNoParent || !bOnlyComposites || bOnlyTasks;
 
-	FBehaviorTreeEditorModule& EditorModule = FModuleManager::GetModuleChecked<FBehaviorTreeEditorModule>(TEXT("BehaviorTreeEditor"));
-	FGraphNodeClassHelper* ClassCache = EditorModule.GetClassCache().Get();
+	FGraphNodeClassHelper& ClassCache = GetClassCache();
 
 	if (bAllowComposites)
 	{
 		FCategorizedGraphActionListBuilder CompositesBuilder(TEXT("Composites"));
 
 		TArray<FGraphNodeClassData> NodeClasses;
-		ClassCache->GatherClasses(UBTCompositeNode::StaticClass(), NodeClasses);
+		ClassCache.GatherClasses(UBTCompositeNode::StaticClass(), NodeClasses);
 
 		const FString ParallelClassName = UBTComposite_SimpleParallel::StaticClass()->GetName();
 		for (const auto& NodeClass : NodeClasses)
 		{
 			const FText NodeTypeName = FText::FromString(FName::NameToDisplayString(NodeClass.ToString(), false));
 
-			TSharedPtr<FAISchemaAction_NewNode> AddOpAction = UAIGraphSchema::AddNewNodeAction(CompositesBuilder, NodeClass.GetCategory(), NodeTypeName, FText::GetEmpty());
+			TSharedPtr<FAISchemaAction_NewNode> AddOpAction = UAIGraphSchema::AddNewNodeAction(CompositesBuilder, NodeClass.GetCategory(), NodeTypeName, NodeClass.GetTooltip());
 
 			UClass* GraphNodeClass = UBehaviorTreeGraphNode_Composite::StaticClass();
 			if (NodeClass.GetClassName() == ParallelClassName)
@@ -156,16 +168,16 @@ void UEdGraphSchema_BehaviorTree::GetGraphContextActions(FGraphContextMenuBuilde
 		FCategorizedGraphActionListBuilder TasksBuilder(TEXT("Tasks"));
 
 		TArray<FGraphNodeClassData> NodeClasses;
-		ClassCache->GatherClasses(UBTTaskNode::StaticClass(), NodeClasses);
+		ClassCache.GatherClasses(UBTTaskNode::StaticClass(), NodeClasses);
 
 		for (const auto& NodeClass : NodeClasses)
 		{
 			const FText NodeTypeName = FText::FromString(FName::NameToDisplayString(NodeClass.ToString(), false));
 
-			TSharedPtr<FAISchemaAction_NewNode> AddOpAction = UAIGraphSchema::AddNewNodeAction(TasksBuilder, NodeClass.GetCategory(), NodeTypeName, FText::GetEmpty());
+			TSharedPtr<FAISchemaAction_NewNode> AddOpAction = UAIGraphSchema::AddNewNodeAction(TasksBuilder, NodeClass.GetCategory(), NodeTypeName, NodeClass.GetTooltip());
 			
 			UClass* GraphNodeClass = UBehaviorTreeGraphNode_Task::StaticClass();
-			if (NodeClass.GetClassName() == UBTTask_RunBehavior::StaticClass()->GetName())
+			if (IsNodeSubtreeTask(NodeClass))
 			{
 				GraphNodeClass = UBehaviorTreeGraphNode_SubtreeTask::StaticClass();
 			}
diff --git a/Engine/Source/Editor/BehaviorTreeEditor/Private/EdGraphSchema_BehaviorTreeDecorator.cpp b/Engine/Source/Editor/BehaviorTreeEditor/Private/EdGraphSchema_BehaviorTreeDecorator.cpp
--- a/Engine/Source/Editor/BehaviorTreeEditor/Private/EdGraphSchema_BehaviorTreeDecorator.cpp
+++ b/Engine/Source/Editor/BehaviorTreeEditor/Private/EdGraphSchema_BehaviorTreeDecorator.cpp
@@ -132,16 +132,23 @@ void UEdGraphSchema_BehaviorTreeDecorator::CreateDefaultNodesForGraph(UEdGraph&
 	NodeCreator.Finalize();
 }
 
+FGraphNodeClassHelper& UEdGraphSchema_BehaviorTreeDecorator::GetClassCache() const
+{
+	const FBehaviorTreeEditorModule& EditorModule = FModuleManager::GetModuleChecked<FBehaviorTreeEditorModule>(TEXT("BehaviorTreeEditor"));
+	FGraphNodeClassHelper* ClassHelper = EditorModule.GetClassCache().Get();
+	check(ClassHelper);
+	return *ClassHelper;
+}
+
 void UEdGraphSchema_BehaviorTreeDecorator::GetGraphContextActions(FGraphContextMenuBuilder& ContextMenuBuilder) const
 {
 	const UBehaviorTreeDecoratorGraphNode* ParentGraphNode = ContextMenuBuilder.FromPin 
 		? Cast<UBehaviorTreeDecoratorGraphNode>(ContextMenuBuilder.FromPin->GetOuter()) : NULL;
 
-	FBehaviorTreeEditorModule& EditorModule = FModuleManager::GetModuleChecked<FBehaviorTreeEditorModule>(TEXT("BehaviorTreeEditor"));
-	FGraphNodeClassHelper* ClassCache = EditorModule.GetClassCache().Get();
+	FGraphNodeClassHelper& ClassCache = GetClassCache();
 
 	TArray<FGraphNodeClassData> NodeClasses;
-	ClassCache->GatherClasses(UBTDecorator::StaticClass(), NodeClasses);
+	ClassCache.GatherClasses(UBTDecorator::StaticClass(), NodeClasses);
 
 	for (const auto& NodeClass : NodeClasses)
 	{
diff --git a/Engine/Source/Editor/BehaviorTreeEditor/Public/AssetDefinition_BehaviorTree.h b/Engine/Source/Editor/BehaviorTreeEditor/Public/AssetDefinition_BehaviorTree.h
new file mode 100644
--- /dev/null
+++ b/Engine/Source/Editor/BehaviorTreeEditor/Public/AssetDefinition_BehaviorTree.h
@@ -0,0 +1,29 @@
+// Copyright Epic Games, Inc. All Rights Reserved.
+
+#pragma once
+
+#include "AssetDefinitionDefault.h"
+#include "AssetDefinition_BehaviorTree.generated.h"
+
+class UBehaviorTree;
+
+UCLASS()
+class BEHAVIORTREEEDITOR_API UAssetDefinition_BehaviorTree : public UAssetDefinitionDefault
+{
+	GENERATED_BODY()
+
+protected:
+	// UAssetDefinition Begin
+	virtual FText GetAssetDisplayName() const override;
+	virtual FLinearColor GetAssetColor() const override;
+	virtual TSoftClassPtr<UObject> GetAssetClass() const override;
+	virtual TConstArrayView<FAssetCategoryPath> GetAssetCategories() const override;
+	virtual EAssetCommandResult OpenAssets(const FAssetOpenArgs& OpenArgs) const override;
+	virtual EAssetCommandResult PerformAssetDiff(const FAssetDiffArgs& DiffArgs) const override;
+	// UAssetDefinition End
+
+private:
+
+	/* Called to open the Behavior Tree defaults view, this opens whatever text diff tool the user has */
+	void OpenInDefaults(class UBehaviorTree* OldBehaviorTree, class UBehaviorTree* NewBehaviorTree) const;
+};
diff --git a/Engine/Source/Editor/BehaviorTreeEditor/Public/BehaviorTreeEditor.h b/Engine/Source/Editor/BehaviorTreeEditor/Public/BehaviorTreeEditor.h
new file mode 100644
--- /dev/null
+++ b/Engine/Source/Editor/BehaviorTreeEditor/Public/BehaviorTreeEditor.h
@@ -0,0 +1,356 @@
+// Copyright Epic Games, Inc. All Rights Reserved.
+
+#pragma once
+
+#include "AIGraphEditor.h"
+#include "Containers/Set.h"
+#include "Containers/UnrealString.h"
+#include "Delegates/IDelegateInstance.h"
+#include "GraphEditor.h"
+#include "HAL/PlatformMath.h"
+#include "IBehaviorTreeEditor.h"
+#include "Internationalization/Text.h"
+#include "Layout/Visibility.h"
+#include "Math/Color.h"
+#include "Misc/NotifyHook.h"
+#include "Templates/SharedPointer.h"
+#include "Toolkits/IToolkit.h"
+#include "Types/SlateEnums.h"
+#include "UObject/NameTypes.h"
+#include "UObject/WeakObjectPtrTemplates.h"
+
+class FDocumentTabFactory;
+class FDocumentTracker;
+class FObjectPostSaveContext;
+class FProperty;
+class IDetailsView;
+class SWidget;
+class UBehaviorTree;
+class UBehaviorTreeComponent;
+class UBehaviorTreeGraph;
+class UBlackboardData;
+class UClass;
+class UEdGraph;
+class UEdGraphNode;
+class UObject;
+class UPackage;
+struct FBlackboardEntry;
+struct FPropertyChangedEvent;
+
+class BEHAVIORTREEEDITOR_API FBehaviorTreeEditor : public IBehaviorTreeEditor, public FAIGraphEditor, public FNotifyHook
+{
+public:
+	FBehaviorTreeEditor();
+	/** Destructor */
+	virtual ~FBehaviorTreeEditor();
+
+	virtual void RegisterTabSpawners(const TSharedRef<class FTabManager>& TabManager) override;
+
+	virtual void InitBehaviorTreeEditor(const EToolkitMode::Type Mode, const TSharedPtr<class IToolkitHost>& InitToolkitHost, UObject* InObject);
+
+	//~ Begin IToolkit Interface
+	virtual FName GetToolkitFName() const override;
+	virtual FText GetBaseToolkitName() const override;
+	virtual FString GetWorldCentricTabPrefix() const override;
+	virtual FLinearColor GetWorldCentricTabColorScale() const override;
+	virtual FText GetToolkitName() const override;
+	virtual FText GetToolkitToolTipText() const override;
+	//~ End IToolkit Interface
+
+	//~ Begin IBehaviorTreeEditor Interface
+	virtual void InitializeDebuggerState(class FBehaviorTreeDebugger* ParentDebugger) const override;
+	virtual UEdGraphNode* FindInjectedNode(int32 Index) const override;
+	virtual void DoubleClickNode(class UEdGraphNode* Node) override;
+	virtual void FocusAttentionOnNode(UEdGraphNode* Node) override;
+	virtual void FocusWindow(UObject* ObjectToFocusOn = NULL) override;
+	//~ End IBehaviorTreeEditor Interface
+
+	//~ Begin FEditorUndoClient Interface
+	virtual void PostUndo(bool bSuccess) override;
+	virtual void PostRedo(bool bSuccess) override;
+	// End of FEditorUndoClient
+
+	//~ Begin FNotifyHook Interface
+	virtual void NotifyPostChange( const FPropertyChangedEvent& PropertyChangedEvent, FProperty* PropertyThatChanged ) override;
+	// End of FNotifyHook
+
+	// Delegates
+	void OnNodeDoubleClicked(class UEdGraphNode* Node);
+	void OnGraphEditorFocused(const TSharedRef<SGraphEditor>& InGraphEditor);
+	void OnNodeTitleCommitted(const FText& NewText, ETextCommit::Type CommitInfo, UEdGraphNode* NodeBeingChanged);
+
+	void OnAddInputPin();
+	bool CanAddInputPin() const;
+	void OnRemoveInputPin();
+	bool CanRemoveInputPin() const;
+
+	void OnEnableBreakpoint();
+	bool CanEnableBreakpoint() const;
+	void OnToggleBreakpoint();
+	bool CanToggleBreakpoint() const;
+	void OnDisableBreakpoint();
+	bool CanDisableBreakpoint() const;
+	void OnAddBreakpoint();
+	bool CanAddBreakpoint() const;
+	void OnRemoveBreakpoint();
+	bool CanRemoveBreakpoint() const;
+
+	void SearchTree();
+	bool CanSearchTree() const;
+
+	void JumpToNode(const UEdGraphNode* Node);
+
+	bool IsPropertyEditable() const;
+	void OnPackageSaved(const FString& PackageFileName, UPackage* Package, FObjectPostSaveContext ObjectSaveContext);
+	void OnFinishedChangingProperties(const FPropertyChangedEvent& PropertyChangedEvent);
+	virtual void OnClassListUpdated() override;
+
+	void UpdateToolbar();
+	bool IsDebuggerReady() const;
+
+	/** Get whether the debugger is currently running and the PIE session is paused */
+	bool IsDebuggerPaused() const;
+
+	/** Get whether we can edit the tree/blackboard with the debugger active */
+	bool CanEditWithDebuggerActive() const;
+
+	TSharedRef<class SWidget> OnGetDebuggerActorsMenu();
+	void OnDebuggerActorSelected(TWeakObjectPtr<UBehaviorTreeComponent> InstanceToDebug);
+	FText GetDebuggerActorDesc() const;
+	FGraphAppearanceInfo GetGraphAppearance() const;
+	bool InEditingMode(bool bGraphIsEditable) const;
+
+	void DebuggerSwitchAsset(UBehaviorTree* NewAsset);
+	void DebuggerUpdateGraph();
+
+	EVisibility GetDebuggerDetailsVisibility() const;
+	EVisibility GetRangeLowerVisibility() const;
+	EVisibility GetRangeSelfVisibility() const;
+	EVisibility GetInjectedNodeVisibility() const;
+	EVisibility GetRootLevelNodeVisibility() const;
+
+	TWeakPtr<SGraphEditor> GetFocusedGraphPtr() const;
+
+	/** Check whether the behavior tree mode can be accessed (i.e whether we have a valid tree to edit) */
+	bool CanAccessBehaviorTreeMode() const;
+
+	/** Check whether the blackboard mode can be accessed (i.e whether we have a valid blackboard to edit) */
+	bool CanAccessBlackboardMode() const;
+
+	/** 
+	 * Get the localized text to display for the specified mode 
+	 * @param	InMode	The mode to display
+	 * @return the localized text representation of the mode
+	 */
+	static FText GetLocalizedMode(FName InMode);
+
+	/** Access the toolbar builder for this editor */
+	TSharedPtr<class FBehaviorTreeEditorToolbar> GetToolbarBuilder() { return ToolbarBuilder; }
+
+	/** Get the behavior tree we are editing (if any) */
+	UBehaviorTree* GetBehaviorTree() const;
+
+	/** Get the blackboard we are editing (if any) */
+	UBlackboardData* GetBlackboardData() const;
+
+	/** Spawns the tab with the update graph inside */
+	TSharedRef<SWidget> SpawnProperties();
+
+	/** Spawns the search tab */
+	TSharedRef<SWidget> SpawnSearch();
+
+	/** Spawn blackboard details tab */
+	TSharedRef<SWidget> SpawnBlackboardDetails();
+
+	/** Spawn blackboard view tab */
+	TSharedRef<SWidget> SpawnBlackboardView();
+
+	/** Spawn blackboard editor tab */
+	TSharedRef<SWidget> SpawnBlackboardEditor();
+
+	// @todo This is a hack for now until we reconcile the default toolbar with application modes [duplicated from counterpart in Blueprint Editor]
+	void RegisterToolbarTab(const TSharedRef<class FTabManager>& TabManager);
+
+	/** Restores the behavior tree graph we were editing or creates a new one if none is available */
+	void RestoreBehaviorTree();
+
+	/** Save the graph state for later editing */
+	void SaveEditedObjectState();
+
+	/** Delegate handler for selection in the blackboard entry list */
+	void HandleBlackboardEntrySelected(const FBlackboardEntry* BlackboardEntry, bool bIsInherited);
+
+	/** Delegate handler used to retrieve current blackboard selection */
+	int32 HandleGetSelectedBlackboardItemIndex(bool& bOutIsInherited);
+
+	/** Delegate handler for displaying debugger values */
+	FText HandleGetDebugKeyValue(const FName& InKeyName, bool bUseCurrentState) const;
+
+	/** Delegate handler for retrieving timestamp to display */
+	double HandleGetDebugTimeStamp(bool bUseCurrentState) const;
+
+	/** Delegate handler for when the debugged blackboard changes */
+	void HandleDebuggedBlackboardChanged(UBlackboardData* InObject);
+
+	/** Delegate handler for determining whether to display the current state */
+	bool HandleGetDisplayCurrentState() const;
+
+	/** Delegate handler for when a blackboard key changes (added, removed, renamed) */
+	void HandleBlackboardKeyChanged(UBlackboardData* InBlackboardData, FBlackboardEntry* const InKey);
+
+	/** Check whether blackboard mode is current */
+	bool HandleIsBlackboardModeActive() const;
+
+	/** Get the currently selected blackboard entry */
+	void GetBlackboardSelectionInfo(int32& OutSelectionIndex, bool& bOutIsInherited) const;
+
+	/** Check to see if we can create a new task node */
+	bool CanCreateNewTask() const;
+
+	/** Check to see if we can create a new decorator node */
+	bool CanCreateNewDecorator() const;
+
+	/** Check to see if we can create a new service node */
+	bool CanCreateNewService() const;
+
+	/** Create the menu used to make a new task node */
+	TSharedRef<SWidget> HandleCreateNewTaskMenu() const;
+
+	/** Create the menu used to make a new decorator */
+	TSharedRef<SWidget> HandleCreateNewDecoratorMenu() const;
+
+	/** Create the menu used to make a new service */
+	TSharedRef<SWidget> HandleCreateNewServiceMenu() const;
+
+	/** Handler for when a node class is picked */
+	void HandleNewNodeClassPicked(UClass* InClass) const;
+
+	/** Create a new task from UBTTask_BlueprintBase */
+	void CreateNewTask() const;
+
+	/** Whether the single button to create a new Blueprint-based task is visible */
+	bool IsNewTaskButtonVisible() const;
+
+	/** Whether the combo button to create a new Blueprint-based task from all available base classes is visible */
+	bool IsNewTaskComboVisible() const;
+
+	/** Create a new decorator from UBTDecorator_BlueprintBase */
+	void CreateNewDecorator() const;
+
+	/** Whether the single button to create a new Blueprint-based decorator is visible */
+	bool IsNewDecoratorButtonVisible() const;
+
+	/** Whether the combo button to create a new Blueprint-based decorator from all available base classes is visible */
+	bool IsNewDecoratorComboVisible() const;
+
+	/** Create a new service from UBTService_BlueprintBase */
+	void CreateNewService() const;
+
+	/** Whether the single button to create a new Blueprint-based service is visible */
+	bool IsNewServiceButtonVisible() const;
+
+	/** Whether the combo button to create a new Blueprint-based service from all available base classes is visible */
+	bool IsNewServiceComboVisible() const;
+
+	/** Create a new Blackboard alongside the currently-edited behavior tree */
+	void CreateNewBlackboard();
+
+	/** Whether we can currently create a new Blackboard */
+	bool CanCreateNewBlackboard() const;
+
+protected:
+	/** Called when "Save" is clicked for this asset */
+	virtual void SaveAsset_Execute() override;
+
+	void SetToolbarCreateActionsEnabled(bool bActionsEnabled);
+
+	TSubclassOf<UBehaviorTreeGraph> GraphClass;
+	FName GraphName;
+	FText CornerText;
+	FText TitleText;
+
+private:
+	/** Create widget for graph editing */
+	TSharedRef<class SGraphEditor> CreateGraphEditorWidget(UEdGraph* InGraph);
+
+	/** Creates all internal widgets for the tabs to point at */
+	void CreateInternalWidgets();
+
+	/** Add custom menu options */
+	void ExtendMenu();
+
+	/** Setup common commands */
+	void BindCommonCommands();
+
+	/** Setup commands */
+	void BindDebuggerToolbarCommands();
+
+	/** Called when the selection changes in the GraphEditor */
+	virtual void OnSelectedNodesChanged(const TSet<class UObject*>& NewSelection) override;
+
+	/** prepare range of nodes that can be aborted by this decorator */
+	void GetAbortModePreview(const class UBehaviorTreeGraphNode_CompositeDecorator* Node, struct FAbortDrawHelper& Mode0, struct FAbortDrawHelper& Mode1);
+
+	/** prepare range of nodes that can be aborted by this decorator */
+	void GetAbortModePreview(const class UBTDecorator* DecoratorOb, struct FAbortDrawHelper& Mode0, struct FAbortDrawHelper& Mode1);
+
+	/** Refresh the debugger's display */
+	void RefreshDebugger();
+
+	/** Push new associated Blackboard data to Blackboard views */
+	void RefreshBlackboardViewsAssociatedObject();
+
+	TSharedPtr<FDocumentTracker> DocumentManager;
+	TWeakPtr<FDocumentTabFactory> GraphEditorTabFactoryPtr;
+
+	/* The Behavior Tree being edited */
+	UBehaviorTree* BehaviorTree;
+
+	/* The Blackboard Data being edited */
+	UBlackboardData* BlackboardData;
+
+	TWeakObjectPtr<class UBehaviorTreeGraphNode_CompositeDecorator> FocusedGraphOwner;
+
+	/** Property View */
+	TSharedPtr<class IDetailsView> DetailsView;
+
+	TSharedPtr<class FBehaviorTreeDebugger> Debugger;
+
+	/** Find results log as well as the search filter */
+	TSharedPtr<class SFindInBT> FindResults;
+
+	uint32 bShowDecoratorRangeLower : 1;
+	uint32 bShowDecoratorRangeSelf : 1;
+	uint32 bForceDisablePropertyEdit : 1;
+	uint32 bSelectedNodeIsInjected : 1;
+	uint32 bSelectedNodeIsRootLevel : 1;
+
+	uint32 bHasMultipleTaskBP : 1;
+	uint32 bHasMultipleDecoratorBP : 1;
+	uint32 bHasMultipleServiceBP : 1;
+
+	TSharedPtr<class FBehaviorTreeEditorToolbar> ToolbarBuilder;
+
+	/** The details view we use to display the blackboard */
+	TSharedPtr<IDetailsView> BlackboardDetailsView;
+
+	/** The blackboard view widget */
+	TSharedPtr<class SBehaviorTreeBlackboardView> BlackboardView;
+
+	/** The blackboard editor widget */
+	TSharedPtr<class SBehaviorTreeBlackboardEditor> BlackboardEditor;
+
+	/** The current blackboard selection index, stored here so it can be accessed by our details customization */
+	int32 CurrentBlackboardEntryIndex;
+
+	/** Whether the current selection is inherited, stored here so it can be accessed by our details customization */
+	bool bIsCurrentBlackboardEntryInherited;
+
+	/** Handle to the registered OnPackageSave delegate */
+	FDelegateHandle OnPackageSavedDelegateHandle;
+
+public:
+	/** Modes in mode switcher */
+	static const FName BehaviorTreeMode;
+	static const FName BlackboardMode;
+};
diff --git a/Engine/Source/Editor/BehaviorTreeEditor/Public/BehaviorTreeEditorModule.h b/Engine/Source/Editor/BehaviorTreeEditor/Public/BehaviorTreeEditorModule.h
--- a/Engine/Source/Editor/BehaviorTreeEditor/Public/BehaviorTreeEditorModule.h
+++ b/Engine/Source/Editor/BehaviorTreeEditor/Public/BehaviorTreeEditorModule.h
@@ -14,7 +14,7 @@ struct FGraphNodeClassHelper;
 DECLARE_LOG_CATEGORY_EXTERN(LogBehaviorTreeEditor, Log, All);
 
 /** Behavior Tree Editor module */
-class FBehaviorTreeEditorModule : public IModuleInterface, public IHasMenuExtensibility, public IHasToolBarExtensibility
+class BEHAVIORTREEEDITOR_API FBehaviorTreeEditorModule : public IModuleInterface, public IHasMenuExtensibility, public IHasToolBarExtensibility
 {
 public:
 	// IModuleInterface
@@ -28,7 +28,7 @@ public:
 	virtual TSharedPtr<FExtensibilityManager> GetMenuExtensibilityManager() override { return MenuExtensibilityManager; }
 	virtual TSharedPtr<FExtensibilityManager> GetToolBarExtensibilityManager() override { return ToolBarExtensibilityManager; }
 
-	TSharedPtr<FGraphNodeClassHelper> GetClassCache() { return ClassCache; }
+	TSharedPtr<FGraphNodeClassHelper> GetClassCache() const { return ClassCache; }
 
 	/** Behavior Tree app identifier string */
 	static const FName BehaviorTreeEditorAppIdentifier;
diff --git a/Engine/Source/Runtime/AIModule/Classes/BehaviorTree/BTNode.h b/Engine/Source/Runtime/AIModule/Classes/BehaviorTree/BTNode.h
--- a/Engine/Source/Runtime/AIModule/Classes/BehaviorTree/BTNode.h
+++ b/Engine/Source/Runtime/AIModule/Classes/BehaviorTree/BTNode.h
@@ -13,6 +13,7 @@
 #include "BTNode.generated.h"
 
 class AActor;
+class ISlateStyle;
 class UBehaviorTree;
 class UBlackboardData;
 class UBTCompositeNode;
diff --git a/Engine/Source/Runtime/Engine/Classes/EdGraph/EdGraph.h b/Engine/Source/Runtime/Engine/Classes/EdGraph/EdGraph.h
--- a/Engine/Source/Runtime/Engine/Classes/EdGraph/EdGraph.h
+++ b/Engine/Source/Runtime/Engine/Classes/EdGraph/EdGraph.h
@@ -279,16 +279,16 @@ public:
 	}
 
 	/** Create an empty placeable graph node */
-	NodeType* CreateNode(bool bSelectNewNode = true)
+	NodeType* CreateNode(bool bSelectNewNode = true, TSubclassOf<NodeType> NodeClass = NodeType::StaticClass())
 	{
-		Node = (NodeType*)Graph.CreateNode(NodeType::StaticClass(), bSelectNewNode);
+		Node = (NodeType*)Graph.CreateNode(NodeClass, bSelectNewNode);
 		return Node;
 	} 
 
 	/** Create an empty placeable graph node */
-	NodeType* CreateUserInvokedNode(bool bSelectNewNode = true)
+	NodeType* CreateUserInvokedNode(bool bSelectNewNode = true, TSubclassOf<NodeType> NodeClass = NodeType::StaticClass())
 	{
-		Node = (NodeType*)Graph.CreateUserInvokedNode(NodeType::StaticClass(), bSelectNewNode);
+		Node = (NodeType*)Graph.CreateUserInvokedNode(NodeClass, bSelectNewNode);
 		return Node;
 	}
 
