Add a Default Blackboard config option for Behavior Trees

================================

Add missing change to CL 25473876.

Add a Default Blackboard config option for Behavior Trees

#[review]-25477192

====================

diff --git a/Engine/Source/Editor/BehaviorTreeEditor/Private/BehaviorTreeGraphNode_Root.cpp b/Engine/Source/Editor/BehaviorTreeEditor/Private/BehaviorTreeGraphNode_Root.cpp
--- a/Engine/Source/Editor/BehaviorTreeEditor/Private/BehaviorTreeGraphNode_Root.cpp
+++ b/Engine/Source/Editor/BehaviorTreeEditor/Private/BehaviorTreeGraphNode_Root.cpp
@@ -2,6 +2,7 @@
 
 #include "BehaviorTreeGraphNode_Root.h"
 
+#include "AISystem.h"
 #include "BehaviorTree/BehaviorTree.h"
 #include "BehaviorTree/BlackboardData.h"
 #include "BehaviorTreeEditorTypes.h"
@@ -23,6 +24,17 @@ void UBehaviorTreeGraphNode_Root::PostPlacedNewNode()
 {
 	Super::PostPlacedNewNode();
 
+	const TSoftObjectPtr<UBlackboardData>& DefaultBlackboard = GET_AI_CONFIG_VAR(DefaultBlackboard);
+	if (DefaultBlackboard.IsValid())
+	{
+		if (UBlackboardData* Blackboard = DefaultBlackboard.LoadSynchronous())
+		{
+			BlackboardAsset = Blackboard;
+			UpdateBlackboard();
+			return;
+		}
+	}
+
 	// pick first available blackboard asset, hopefully something will be loaded...
 	for (FThreadSafeObjectIterator It(UBlackboardData::StaticClass()); It; ++It)
 	{
diff --git a/Engine/Source/Runtime/AIModule/Classes/AISystem.h b/Engine/Source/Runtime/AIModule/Classes/AISystem.h
--- a/Engine/Source/Runtime/AIModule/Classes/AISystem.h
+++ b/Engine/Source/Runtime/AIModule/Classes/AISystem.h
@@ -95,6 +95,10 @@ public:
 	UPROPERTY(globalconfig, EditDefaultsOnly, Category = "Behavior Tree")
 	bool bBlackboardKeyDecoratorAllowsNoneAsValue = false;
 
+	/** If set, new BTs will use this BB as default. */
+	UPROPERTY(globalconfig, EditDefaultsOnly, Category = "Behavior Tree")
+	TSoftObjectPtr<UBlackboardData> DefaultBlackboard;
+
 	/** Which collision channel to use for sight checks by default */
 	UPROPERTY(globalconfig, EditDefaultsOnly, Category = "PerceptionSystem")
 	TEnumAsByte<ECollisionChannel> DefaultSightCollisionChannel;
