PR #10345: Fix BehaviorTreeEditorModule custom Toolbar Extenders not being applied
[RN] none

diff --git a/Engine/Source/Editor/BehaviorTreeEditor/Private/BehaviorTreeEditor.cpp b/Engine/Source/Editor/BehaviorTreeEditor/Private/BehaviorTreeEditor.cpp
--- a/Engine/Source/Editor/BehaviorTreeEditor/Private/BehaviorTreeEditor.cpp
+++ b/Engine/Source/Editor/BehaviorTreeEditor/Private/BehaviorTreeEditor.cpp
@@ -287,6 +287,7 @@ void FBehaviorTreeEditor::InitBehaviorTreeEditor( const EToolkitMode::Type Mode,
 
 		FBehaviorTreeEditorModule& BehaviorTreeEditorModule = FModuleManager::LoadModuleChecked<FBehaviorTreeEditorModule>( "BehaviorTreeEditor" );
 		AddMenuExtender(BehaviorTreeEditorModule.GetMenuExtensibilityManager()->GetAllExtenders(GetToolkitCommands(), GetEditingObjects()));
+		AddToolbarExtender(BehaviorTreeEditorModule.GetToolBarExtensibilityManager()->GetAllExtenders(GetToolkitCommands(), GetEditingObjects()));
 
 		AddApplicationMode(BehaviorTreeMode, MakeShareable(new FBehaviorTreeEditorApplicationMode(SharedThis(this))));
 		AddApplicationMode(BlackboardMode, MakeShareable(new FBlackboardEditorApplicationMode(SharedThis(this))));
