Add custom style set support for BT nodes

diff --git a/Engine/Source/Editor/BehaviorTreeEditor/Classes/BehaviorTreeGraphNode.h b/Engine/Source/Editor/BehaviorTreeEditor/Classes/BehaviorTreeGraphNode.h
--- a/Engine/Source/Editor/BehaviorTreeEditor/Classes/BehaviorTreeGraphNode.h
+++ b/Engine/Source/Editor/BehaviorTreeEditor/Classes/BehaviorTreeGraphNode.h
@@ -21,7 +21,7 @@ class UObject;
 template <typename T> struct TObjectPtr;
 
 UCLASS()
-class UBehaviorTreeGraphNode : public UAIGraphNode
+class BEHAVIORTREEEDITOR_API UBehaviorTreeGraphNode : public UAIGraphNode
 {
 	GENERATED_UCLASS_BODY()
 
@@ -59,6 +59,10 @@ class UBehaviorTreeGraphNode : public UAIGraphNode
 
 	void ClearDebuggerState();
 
+	// #RED_BTNodeCustomIconStyleSet kristof.morva 2023-04-07
+	/** gets the style set from which GetNameIcon is queried */
+	virtual const ISlateStyle& GetNameIconStyleSet() const;
+
 	/** gets icon resource name for title bar */
 	virtual FName GetNameIcon() const;
 
diff --git a/Engine/Source/Editor/BehaviorTreeEditor/Private/BehaviorTreeGraphNode.cpp b/Engine/Source/Editor/BehaviorTreeEditor/Private/BehaviorTreeGraphNode.cpp
--- a/Engine/Source/Editor/BehaviorTreeEditor/Private/BehaviorTreeGraphNode.cpp
+++ b/Engine/Source/Editor/BehaviorTreeEditor/Private/BehaviorTreeGraphNode.cpp
@@ -390,6 +390,14 @@ void UBehaviorTreeGraphNode::ClearDebuggerState()
 	DebuggerRuntimeDescription.Empty();
 }
 
+// < #RED_BTNodeCustomIconStyleSet kristof.morva 2023-04-07
+const ISlateStyle& UBehaviorTreeGraphNode::GetNameIconStyleSet() const
+{
+	const UBTNode* BTNodeInstance = Cast<UBTNode>(NodeInstance);
+	return BTNodeInstance != nullptr ? BTNodeInstance->GetNodeIconStyleSet() : FAppStyle::Get();
+}
+// > #RED_BTNodeCustomIconStyleSet kristof.morva 2023-04-07
+
 FName UBehaviorTreeGraphNode::GetNameIcon() const
 {
 	UBTNode* BTNodeInstance = Cast<UBTNode>(NodeInstance);
diff --git a/Engine/Source/Editor/BehaviorTreeEditor/Private/SGraphNode_BehaviorTree.cpp b/Engine/Source/Editor/BehaviorTreeEditor/Private/SGraphNode_BehaviorTree.cpp
--- a/Engine/Source/Editor/BehaviorTreeEditor/Private/SGraphNode_BehaviorTree.cpp
+++ b/Engine/Source/Editor/BehaviorTreeEditor/Private/SGraphNode_BehaviorTree.cpp
@@ -869,7 +869,13 @@ TSharedPtr<SToolTip> SGraphNode_BehaviorTree::GetComplexTooltip()
 const FSlateBrush* SGraphNode_BehaviorTree::GetNameIcon() const
 {	
 	UBehaviorTreeGraphNode* BTGraphNode = Cast<UBehaviorTreeGraphNode>(GraphNode);
-	return BTGraphNode != nullptr ? FAppStyle::GetBrush(BTGraphNode->GetNameIcon()) : FAppStyle::GetBrush(TEXT("BTEditor.Graph.BTNode.Icon"));
+	if (BTGraphNode != nullptr)
+	{
+		// #RED_BTNodeCustomIconStyleSet kristof.morva 2023-04-07
+		return BTGraphNode->GetNameIconStyleSet().GetBrush(BTGraphNode->GetNameIcon());
+	}
+
+	return FAppStyle::GetBrush(TEXT("BTEditor.Graph.BTNode.Icon"));
 }
 
 static UBehaviorTreeGraphNode* GetParentNode(UEdGraphNode* GraphNode)
diff --git a/Engine/Source/Runtime/AIModule/AIModule.Build.cs b/Engine/Source/Runtime/AIModule/AIModule.Build.cs
--- a/Engine/Source/Runtime/AIModule/AIModule.Build.cs
+++ b/Engine/Source/Runtime/AIModule/AIModule.Build.cs
@@ -30,6 +30,7 @@ namespace UnrealBuildTool.Rules
             {
 				PrivateDependencyModuleNames.Add("EditorFramework");
                 PrivateDependencyModuleNames.Add("UnrealEd");
+                PrivateDependencyModuleNames.Add("SlateCore");
 
                 PrivateDependencyModuleNames.Add("AITestSuite");
                 CircularlyReferencedDependentModules.Add("AITestSuite");
diff --git a/Engine/Source/Runtime/AIModule/Classes/BehaviorTree/BTNode.h b/Engine/Source/Runtime/AIModule/Classes/BehaviorTree/BTNode.h
--- a/Engine/Source/Runtime/AIModule/Classes/BehaviorTree/BTNode.h
+++ b/Engine/Source/Runtime/AIModule/Classes/BehaviorTree/BTNode.h
@@ -148,6 +148,9 @@ class UBTNode : public UObject, public IGameplayTaskOwnerInterface
 	AIMODULE_API virtual FString GetStaticDescription() const;
 
 #if WITH_EDITOR
+	/** Get the style set from which GetNodeIconName is queried */
+	AIMODULE_API virtual const ISlateStyle& GetNodeIconStyleSet() const;
+
 	/** Get the name of the icon used to display this node in the editor */
 	AIMODULE_API virtual FName GetNodeIconName() const;
 
diff --git a/Engine/Source/Runtime/AIModule/Private/BehaviorTree/BTNode.cpp b/Engine/Source/Runtime/AIModule/Private/BehaviorTree/BTNode.cpp
--- a/Engine/Source/Runtime/AIModule/Private/BehaviorTree/BTNode.cpp
+++ b/Engine/Source/Runtime/AIModule/Private/BehaviorTree/BTNode.cpp
@@ -205,6 +205,11 @@ void UBTNode::DescribeRuntimeValues(const UBehaviorTreeComponent& OwnerComp, uin
 
 #if WITH_EDITOR
 
+const ISlateStyle& UBTNode::GetNodeIconStyleSet() const
+{
+	return FAppStyle::Get();
+}
+
 FName UBTNode::GetNodeIconName() const
 {
 	return NAME_None;
