remove misc. unnecessary comments

#rnx

diff --git a/Engine/Source/Editor/BehaviorTreeEditor/Classes/BehaviorTreeGraphNode.h b/Engine/Source/Editor/BehaviorTreeEditor/Classes/BehaviorTreeGraphNode.h
--- a/Engine/Source/Editor/BehaviorTreeEditor/Classes/BehaviorTreeGraphNode.h
+++ b/Engine/Source/Editor/BehaviorTreeEditor/Classes/BehaviorTreeGraphNode.h
@@ -59,7 +59,6 @@ class BEHAVIORTREEEDITOR_API UBehaviorTreeGraphNode : public UAIGraphNode
 
 	void ClearDebuggerState();
 
-	// #RED_BTNodeCustomIconStyleSet kristof.morva 2023-04-07
 	/** gets the style set from which GetNameIcon is queried */
 	virtual const ISlateStyle& GetNameIconStyleSet() const;
 
diff --git a/Engine/Source/Editor/BehaviorTreeEditor/Private/BehaviorTreeGraphNode.cpp b/Engine/Source/Editor/BehaviorTreeEditor/Private/BehaviorTreeGraphNode.cpp
--- a/Engine/Source/Editor/BehaviorTreeEditor/Private/BehaviorTreeGraphNode.cpp
+++ b/Engine/Source/Editor/BehaviorTreeEditor/Private/BehaviorTreeGraphNode.cpp
@@ -390,13 +390,11 @@ void UBehaviorTreeGraphNode::ClearDebuggerState()
 	DebuggerRuntimeDescription.Empty();
 }
 
-// < #RED_BTNodeCustomIconStyleSet kristof.morva 2023-04-07
 const ISlateStyle& UBehaviorTreeGraphNode::GetNameIconStyleSet() const
 {
 	const UBTNode* BTNodeInstance = Cast<UBTNode>(NodeInstance);
 	return BTNodeInstance != nullptr ? BTNodeInstance->GetNodeIconStyleSet() : FAppStyle::Get();
 }
-// > #RED_BTNodeCustomIconStyleSet kristof.morva 2023-04-07
 
 FName UBehaviorTreeGraphNode::GetNameIcon() const
 {
diff --git a/Engine/Source/Editor/BehaviorTreeEditor/Private/SGraphNode_BehaviorTree.cpp b/Engine/Source/Editor/BehaviorTreeEditor/Private/SGraphNode_BehaviorTree.cpp
--- a/Engine/Source/Editor/BehaviorTreeEditor/Private/SGraphNode_BehaviorTree.cpp
+++ b/Engine/Source/Editor/BehaviorTreeEditor/Private/SGraphNode_BehaviorTree.cpp
@@ -871,7 +871,6 @@ const FSlateBrush* SGraphNode_BehaviorTree::GetNameIcon() const
 	UBehaviorTreeGraphNode* BTGraphNode = Cast<UBehaviorTreeGraphNode>(GraphNode);
 	if (BTGraphNode != nullptr)
 	{
-		// #RED_BTNodeCustomIconStyleSet kristof.morva 2023-04-07
 		return BTGraphNode->GetNameIconStyleSet().GetBrush(BTGraphNode->GetNameIcon());
 	}
 
