[BehaviorTreeEditor] converted OnBlackboardKeyChanged delegate to a global multicast delegate so every BehaviorTreeBlackboardView linked to the modified asset can refresh their actions.

diff --git a/Engine/Source/Editor/BehaviorTreeEditor/Private/BehaviorTreeEditor.cpp b/Engine/Source/Editor/BehaviorTreeEditor/Private/BehaviorTreeEditor.cpp
--- a/Engine/Source/Editor/BehaviorTreeEditor/Private/BehaviorTreeEditor.cpp
+++ b/Engine/Source/Editor/BehaviorTreeEditor/Private/BehaviorTreeEditor.cpp
@@ -282,7 +282,6 @@ void FBehaviorTreeEditor::InitBehaviorTreeEditor( const EToolkitMode::Type Mode,
 
 		Debugger = MakeShareable(new FBehaviorTreeDebugger);
 		Debugger->Setup(BehaviorTree, SharedThis(this));
-		Debugger->OnDebuggedBlackboardChanged().AddSP(this, &FBehaviorTreeEditor::HandleDebuggedBlackboardChanged);
 		BindDebuggerToolbarCommands();
 
 		FBehaviorTreeEditorModule& BehaviorTreeEditorModule = FModuleManager::LoadModuleChecked<FBehaviorTreeEditorModule>( "BehaviorTreeEditor" );
@@ -306,7 +305,6 @@ void FBehaviorTreeEditor::InitBehaviorTreeEditor( const EToolkitMode::Type Mode,
 				.OnIsDebuggerPaused(this, &FBehaviorTreeEditor::IsDebuggerPaused)
 				.OnGetDebugTimeStamp(this, &FBehaviorTreeEditor::HandleGetDebugTimeStamp)
 				.OnGetDisplayCurrentState(this, &FBehaviorTreeEditor::HandleGetDisplayCurrentState)
-				.OnBlackboardKeyChanged(this, &FBehaviorTreeEditor::HandleBlackboardKeyChanged)
 				.OnIsBlackboardModeActive(this, &FBehaviorTreeEditor::HandleIsBlackboardModeActive);
 	}
 	else
@@ -456,15 +454,6 @@ bool FBehaviorTreeEditor::HandleGetDisplayCurrentState() const
 	return false;
 }
 
-void FBehaviorTreeEditor::HandleBlackboardKeyChanged(UBlackboardData* InBlackboardData, FBlackboardEntry* const InKey)
-{
-	if(BlackboardView.IsValid())
-	{
-		// re-set object in blackboard view to keep it up to date
-		BlackboardView->SetObject(InBlackboardData);
-	}
-}
-
 bool FBehaviorTreeEditor::HandleIsBlackboardModeActive() const
 {
 	return GetCurrentMode() == BlackboardMode;
diff --git a/Engine/Source/Editor/BehaviorTreeEditor/Private/BehaviorTreeEditorDelegates.cpp b/Engine/Source/Editor/BehaviorTreeEditor/Private/BehaviorTreeEditorDelegates.cpp
new file mode 100644
--- /dev/null
+++ b/Engine/Source/Editor/BehaviorTreeEditor/Private/BehaviorTreeEditorDelegates.cpp
@@ -0,0 +1,8 @@
+// Copyright Epic Games, Inc. All Rights Reserved.
+
+#include "BehaviorTreeEditorDelegates.h"
+
+namespace UE::BehaviorTreeEditor::Delegates
+{
+FOnBlackboardKeyChanged OnBlackboardKeyChanged;
+}; // UE::BehaviorTreeEditor::Delegates
diff --git a/Engine/Source/Editor/BehaviorTreeEditor/Private/BehaviorTreeEditorDelegates.h b/Engine/Source/Editor/BehaviorTreeEditor/Private/BehaviorTreeEditorDelegates.h
new file mode 100644
--- /dev/null
+++ b/Engine/Source/Editor/BehaviorTreeEditor/Private/BehaviorTreeEditorDelegates.h
@@ -0,0 +1,15 @@
+// Copyright Epic Games, Inc. All Rights Reserved.
+
+#pragma once
+
+#include "Delegates/Delegate.h"
+
+struct FBlackboardEntry;
+class UBlackboardData;
+
+namespace UE::BehaviorTreeEditor::Delegates
+{
+/** Delegate for when a blackboard key changes (added, removed, renamed) */
+DECLARE_MULTICAST_DELEGATE_TwoParams(FOnBlackboardKeyChanged, const UBlackboardData& /*InBlackboardData*/, FBlackboardEntry* const /*InKey*/);
+extern FOnBlackboardKeyChanged OnBlackboardKeyChanged;
+} // UE::BehaviorTreeEditor::Delegates
diff --git a/Engine/Source/Editor/BehaviorTreeEditor/Private/SBehaviorTreeBlackboardEditor.cpp b/Engine/Source/Editor/BehaviorTreeEditor/Private/SBehaviorTreeBlackboardEditor.cpp
--- a/Engine/Source/Editor/BehaviorTreeEditor/Private/SBehaviorTreeBlackboardEditor.cpp
+++ b/Engine/Source/Editor/BehaviorTreeEditor/Private/SBehaviorTreeBlackboardEditor.cpp
@@ -5,6 +5,7 @@
 #include "BehaviorTree/Blackboard/BlackboardKeyType.h"
 #include "BehaviorTree/BlackboardData.h"
 #include "BehaviorTreeEditorCommands.h"
+#include "BehaviorTreeEditorDelegates.h"
 #include "ClassViewerFilter.h"
 #include "ClassViewerModule.h"
 #include "Containers/Array.h"
@@ -83,7 +84,6 @@ void SBehaviorTreeBlackboardEditor::Construct(const FArguments& InArgs, TSharedR
 		.OnIsDebuggerReady(InArgs._OnIsDebuggerReady)
 		.OnIsDebuggerPaused(InArgs._OnIsDebuggerPaused)
 		.OnGetDebugTimeStamp(InArgs._OnGetDebugTimeStamp)
-		.OnBlackboardKeyChanged(InArgs._OnBlackboardKeyChanged)
 		.IsReadOnly(false),
 		CommandList,
 		InBlackboardData
@@ -153,7 +153,7 @@ void SBehaviorTreeBlackboardEditor::HandleDeleteEntry()
 			}
 
 			GraphActionMenu->RefreshAllActions(true);
-			OnBlackboardKeyChanged.ExecuteIfBound(BlackboardDataPtr, nullptr);
+			UE::BehaviorTreeEditor::Delegates::OnBlackboardKeyChanged.Broadcast(*BlackboardDataPtr, nullptr);
 
 			// signal de-selection
 			if(OnEntrySelected.IsBound())
@@ -313,7 +313,7 @@ void SBehaviorTreeBlackboardEditor::HandleKeyClassPicked(UClass* InClass)
 	BlackboardDataPtr->Keys.Add(Entry);
 
 	GraphActionMenu->RefreshAllActions(true);
-	OnBlackboardKeyChanged.ExecuteIfBound(BlackboardDataPtr, &BlackboardDataPtr->Keys.Last());
+	UE::BehaviorTreeEditor::Delegates::OnBlackboardKeyChanged.Broadcast(*BlackboardDataPtr, &BlackboardDataPtr->Keys.Last());
 
 	GraphActionMenu->SelectItemByName(Entry.EntryName, ESelectInfo::OnMouseClick);
 
diff --git a/Engine/Source/Editor/BehaviorTreeEditor/Private/SBehaviorTreeBlackboardEditor.h b/Engine/Source/Editor/BehaviorTreeEditor/Private/SBehaviorTreeBlackboardEditor.h
--- a/Engine/Source/Editor/BehaviorTreeEditor/Private/SBehaviorTreeBlackboardEditor.h
+++ b/Engine/Source/Editor/BehaviorTreeEditor/Private/SBehaviorTreeBlackboardEditor.h
@@ -34,7 +34,6 @@ public:
 		SLATE_EVENT(FOnIsDebuggerReady, OnIsDebuggerReady)
 		SLATE_EVENT(FOnIsDebuggerPaused, OnIsDebuggerPaused)
 		SLATE_EVENT(FOnGetDebugTimeStamp, OnGetDebugTimeStamp)
-		SLATE_EVENT(FOnBlackboardKeyChanged, OnBlackboardKeyChanged)
 		SLATE_EVENT(FOnIsBlackboardModeActive, OnIsBlackboardModeActive)
 
 	SLATE_END_ARGS()
diff --git a/Engine/Source/Editor/BehaviorTreeEditor/Private/SBehaviorTreeBlackboardView.cpp b/Engine/Source/Editor/BehaviorTreeEditor/Private/SBehaviorTreeBlackboardView.cpp
--- a/Engine/Source/Editor/BehaviorTreeEditor/Private/SBehaviorTreeBlackboardView.cpp
+++ b/Engine/Source/Editor/BehaviorTreeEditor/Private/SBehaviorTreeBlackboardView.cpp
@@ -2,6 +2,7 @@
 
 #include "SBehaviorTreeBlackboardView.h"
 
+#include "BehaviorTreeEditorDelegates.h"
 #include "AssetRegistry/AssetData.h"
 #include "AssetRegistry/AssetRegistryModule.h"
 #include "AssetRegistry/IAssetRegistry.h"
@@ -101,7 +102,6 @@ class SBehaviorTreeBlackboardItem : public SGraphPaletteItem
 		SLATE_EVENT(FOnGetDebugKeyValue, OnGetDebugKeyValue)
 		SLATE_EVENT(FOnGetDisplayCurrentState, OnGetDisplayCurrentState)
 		SLATE_EVENT(FOnIsDebuggerReady, OnIsDebuggerReady)
-		SLATE_EVENT(FOnBlackboardKeyChanged, OnBlackboardKeyChanged)
 
 	SLATE_END_ARGS()
 
@@ -110,7 +110,6 @@ class SBehaviorTreeBlackboardItem : public SGraphPaletteItem
 		OnGetDebugKeyValue = InArgs._OnGetDebugKeyValue;
 		OnIsDebuggerReady = InArgs._OnIsDebuggerReady;
 		OnGetDisplayCurrentState = InArgs._OnGetDisplayCurrentState;
-		OnBlackboardKeyChanged = InArgs._OnBlackboardKeyChanged;
 
 		const FSlateFontInfo NameFont = FCoreStyle::GetDefaultFontStyle("Regular", 10);
 
@@ -256,7 +255,7 @@ private:
 
 				BlackboardEntryAction->Update();
 
-				OnBlackboardKeyChanged.ExecuteIfBound(BlackboardData, &BlackboardEntryAction->Key);
+				UE::BehaviorTreeEditor::Delegates::OnBlackboardKeyChanged.Broadcast(*BlackboardData, &BlackboardEntryAction->Key);
 
 				if (!BlackboardEntryAction->bIsNew)
 				{
@@ -436,9 +435,6 @@ private:
 	/** Delegate used to determine whether the BT debugger displaying the current state */
 	FOnGetDisplayCurrentState OnGetDisplayCurrentState;
 
-	/** Delegate for when a blackboard key changes (added, removed, renamed) */
-	FOnBlackboardKeyChanged OnBlackboardKeyChanged;
-
 	/** Read-only flag */
 	bool bIsReadOnly;
 };
@@ -452,7 +448,8 @@ void SBehaviorTreeBlackboardView::Construct(const FArguments& InArgs, TSharedRef
 	OnIsDebuggerPaused = InArgs._OnIsDebuggerPaused;
 	OnGetDebugTimeStamp = InArgs._OnGetDebugTimeStamp;
 	OnGetDisplayCurrentState = InArgs._OnGetDisplayCurrentState;
-	OnBlackboardKeyChanged = InArgs._OnBlackboardKeyChanged;
+
+	OnBlackboardKeyChangedDelegateHandle = UE::BehaviorTreeEditor::Delegates::OnBlackboardKeyChanged.AddSP(this, &SBehaviorTreeBlackboardView::HandleBlackboardKeyChanged);
 
 	BlackboardData = InBlackboardData;
 
@@ -535,13 +532,17 @@ void SBehaviorTreeBlackboardView::Construct(const FArguments& InArgs, TSharedRef
 	];
 }
 
+SBehaviorTreeBlackboardView::~SBehaviorTreeBlackboardView()
+{
+	UE::BehaviorTreeEditor::Delegates::OnBlackboardKeyChanged.Remove(OnBlackboardKeyChangedDelegateHandle);
+}
+
 TSharedRef<SWidget> SBehaviorTreeBlackboardView::HandleCreateWidgetForAction(FCreateWidgetForActionData* const InCreateData)
 {
 	return SNew(SBehaviorTreeBlackboardItem, InCreateData)
 		.OnIsDebuggerReady(OnIsDebuggerReady)
 		.OnGetDebugKeyValue(OnGetDebugKeyValue)
-		.OnGetDisplayCurrentState(this, &SBehaviorTreeBlackboardView::IsUsingCurrentValues)
-		.OnBlackboardKeyChanged(OnBlackboardKeyChanged);
+		.OnGetDisplayCurrentState(this, &SBehaviorTreeBlackboardView::IsUsingCurrentValues);
 }
 
 void SBehaviorTreeBlackboardView::HandleCollectAllActions( FGraphActionListBuilderBase& GraphActionListBuilder )
@@ -644,6 +645,15 @@ void SBehaviorTreeBlackboardView::SetObject(UBlackboardData* InBlackboardData)
 	GraphActionMenu->RefreshAllActions(true);
 }
 
+void SBehaviorTreeBlackboardView::HandleBlackboardKeyChanged(const UBlackboardData& InBlackboardData, FBlackboardEntry* const InKey)
+{
+	// OnBlackboardKeyChanged is a global delegate so refresh only if it's our asset.
+	if (BlackboardData == &InBlackboardData)
+	{
+		GraphActionMenu->RefreshAllActions(true);
+	}
+}
+
 TSharedPtr<SWidget> SBehaviorTreeBlackboardView::HandleContextMenuOpening(TSharedRef<FUICommandList> ToolkitCommands) const
 {
 	FMenuBuilder MenuBuilder(/* bInShouldCloseWindowAfterMenuSelection =*/true, ToolkitCommands);
diff --git a/Engine/Source/Editor/BehaviorTreeEditor/Private/SBehaviorTreeBlackboardView.h b/Engine/Source/Editor/BehaviorTreeEditor/Private/SBehaviorTreeBlackboardView.h
--- a/Engine/Source/Editor/BehaviorTreeEditor/Private/SBehaviorTreeBlackboardView.h
+++ b/Engine/Source/Editor/BehaviorTreeEditor/Private/SBehaviorTreeBlackboardView.h
@@ -44,8 +44,6 @@ DECLARE_DELEGATE_RetVal(bool, FOnGetDisplayCurrentState);
 /** Delegate used to get the debugger's current timestamp */
 DECLARE_DELEGATE_RetVal_OneParam(double, FOnGetDebugTimeStamp, bool /* bUseCurrentState */);
 
-/** Delegate for when a blackboard key changes (added, removed, renamed) */
-DECLARE_DELEGATE_TwoParams(FOnBlackboardKeyChanged, UBlackboardData* /*InBlackboardData*/, FBlackboardEntry* const /*InKey*/);
 
 /** Blackboard entry in the list */
 class FEdGraphSchemaAction_BlackboardEntry : public FEdGraphSchemaAction_Dummy
@@ -86,12 +84,12 @@ public:
 		SLATE_EVENT(FOnIsDebuggerReady, OnIsDebuggerReady)
 		SLATE_EVENT(FOnIsDebuggerPaused, OnIsDebuggerPaused)
 		SLATE_EVENT(FOnGetDebugTimeStamp, OnGetDebugTimeStamp)
-		SLATE_EVENT(FOnBlackboardKeyChanged, OnBlackboardKeyChanged)
 		SLATE_ARGUMENT(bool, IsReadOnly)
 
 	SLATE_END_ARGS()
 
 	void Construct(const FArguments& InArgs, TSharedRef<FUICommandList> InCommandList, UBlackboardData* InBlackboardData);
+	~SBehaviorTreeBlackboardView();
 
 	/**
 	 * Retrieves the blackboard item currently selected by the user.
@@ -111,6 +109,9 @@ public:
 	void SetObject(UBlackboardData* InBlackboardData);
 
 protected:
+	/** Delegate handler for when a blackboard key changes (added, removed, renamed) */
+	void HandleBlackboardKeyChanged(const UBlackboardData& InBlackboardData, FBlackboardEntry* const InKey);
+
 	/** Delegate handler used to generate a widget for an 'action' (key) in the list */
 	TSharedRef<SWidget> HandleCreateWidgetForAction(FCreateWidgetForActionData* const InCreateData);
 
@@ -190,8 +191,8 @@ protected:
 	/** Delegate used to get the debugger's current timestamp */
 	FOnGetDebugTimeStamp OnGetDebugTimeStamp;
 
-	/** Delegate for when a blackboard key changes (added, removed, renamed) */
-	FOnBlackboardKeyChanged OnBlackboardKeyChanged;
+	/** Handle to the registered OnBlackboardKeyChanged delegate */
+	FDelegateHandle OnBlackboardKeyChangedDelegateHandle;
 
 	/** Whether we want to show the current or saved state */
 	bool bShowCurrentState;
diff --git a/Engine/Source/Editor/BehaviorTreeEditor/Public/BehaviorTreeEditor.h b/Engine/Source/Editor/BehaviorTreeEditor/Public/BehaviorTreeEditor.h
--- a/Engine/Source/Editor/BehaviorTreeEditor/Public/BehaviorTreeEditor.h
+++ b/Engine/Source/Editor/BehaviorTreeEditor/Public/BehaviorTreeEditor.h
@@ -195,9 +195,6 @@ public:
 	/** Delegate handler for determining whether to display the current state */
 	bool HandleGetDisplayCurrentState() const;
 
-	/** Delegate handler for when a blackboard key changes (added, removed, renamed) */
-	void HandleBlackboardKeyChanged(UBlackboardData* InBlackboardData, FBlackboardEntry* const InKey);
-
 	/** Check whether blackboard mode is current */
 	bool HandleIsBlackboardModeActive() const;
 
