[BehaviorTreeEditor] Fix potential access to an invalid FindInBTResult pointer

diff --git a/Engine/Source/Editor/BehaviorTreeEditor/Private/FindInBT.cpp b/Engine/Source/Editor/BehaviorTreeEditor/Private/FindInBT.cpp
--- a/Engine/Source/Editor/BehaviorTreeEditor/Private/FindInBT.cpp
+++ b/Engine/Source/Editor/BehaviorTreeEditor/Private/FindInBT.cpp
@@ -89,17 +89,23 @@ TSharedRef<SWidget> FFindInBTResult::CreateIcon() const
 		.ColorAndOpacity(IconColor);
 }
 
-FReply FFindInBTResult::OnClick(TWeakPtr<class FBehaviorTreeEditor> BehaviorTreeEditorPtr, TSharedPtr<FFindInBTResult> Root)
+FReply FFindInBTResult::OnClick(TWeakPtr<FBehaviorTreeEditor> BehaviorTreeEditor, TSharedPtr<FFindInBTResult> Root)
 {
-	if (BehaviorTreeEditorPtr.IsValid() && GraphNode.IsValid())
+	const TSharedPtr<FBehaviorTreeEditor> BTEditorAsShared = BehaviorTreeEditor.Pin();
+	FBehaviorTreeEditor* BTEditorPtr = BTEditorAsShared.Get();
+
+	const TSharedPtr<FFindInBTResult> ParentAsShared = Parent.Pin();
+	const FFindInBTResult* ParentPtr = ParentAsShared.Get();
+
+	if (BTEditorPtr != nullptr && ParentPtr != nullptr)
 	{
-		if (Parent.IsValid() && Parent.HasSameObject(Root.Get()))
+		if (ParentAsShared == Root)
 		{
-			BehaviorTreeEditorPtr.Pin()->JumpToNode(GraphNode.Get());
+			BTEditorPtr->JumpToNode(GraphNode.Get());
 		}
 		else
 		{
-			BehaviorTreeEditorPtr.Pin()->JumpToNode(Parent.Pin()->GraphNode.Get());
+			BTEditorPtr->JumpToNode(ParentPtr->GraphNode.Get());
 		}
 	}
 
