Export BehaviorTreeGraphNode_SubtreeTask

diff --git a/Engine/Source/Editor/BehaviorTreeEditor/Classes/BehaviorTreeGraphNode_SubtreeTask.h b/Engine/Source/Editor/BehaviorTreeEditor/Classes/BehaviorTreeGraphNode_SubtreeTask.h
--- a/Engine/Source/Editor/BehaviorTreeEditor/Classes/BehaviorTreeGraphNode_SubtreeTask.h
+++ b/Engine/Source/Editor/BehaviorTreeEditor/Classes/BehaviorTreeGraphNode_SubtreeTask.h
@@ -13,7 +13,7 @@
 
 class UObject;
 
-UCLASS()
+UCLASS(MinimalAPI)
 class UBehaviorTreeGraphNode_SubtreeTask : public UBehaviorTreeGraphNode_Task
 {
 	GENERATED_UCLASS_BODY()
