Next pass on Behavior Tree extensibility:
- Customizable node colors
- Inheritable task / decorator / service nodes
- More customizable text
- Add ability to hide services
- Const-correctness fixes

diff --git a/Engine/Source/Editor/BehaviorTreeEditor/Classes/BehaviorTreeGraph.h b/Engine/Source/Editor/BehaviorTreeEditor/Classes/BehaviorTreeGraph.h
--- a/Engine/Source/Editor/BehaviorTreeEditor/Classes/BehaviorTreeGraph.h
+++ b/Engine/Source/Editor/BehaviorTreeEditor/Classes/BehaviorTreeGraph.h
@@ -47,6 +47,8 @@ class BEHAVIORTREEEDITOR_API UBehaviorTreeGraph : public UAIGraph
 	virtual void UpdateAsset(int32 UpdateFlags = 0) override;
 	virtual void OnSubNodeDropped() override;
 
+	virtual bool DoesSupportServices() const { return true; }
+
 	void UpdateBlackboardChange();
 	void UpdateAbortHighlight(struct FAbortDrawHelper& Mode0, struct FAbortDrawHelper& Mode1);
 	void CreateBTFromGraph(class UBehaviorTreeGraphNode* RootEdNode);
diff --git a/Engine/Source/Editor/BehaviorTreeEditor/Classes/BehaviorTreeGraphNode.h b/Engine/Source/Editor/BehaviorTreeEditor/Classes/BehaviorTreeGraphNode.h
--- a/Engine/Source/Editor/BehaviorTreeEditor/Classes/BehaviorTreeGraphNode.h
+++ b/Engine/Source/Editor/BehaviorTreeEditor/Classes/BehaviorTreeGraphNode.h
@@ -15,6 +15,7 @@
 #include "BehaviorTreeGraphNode.generated.h"
 
 class ISlateStyle;
+class UBehaviorTreeGraph;
 class UEdGraph;
 class UEdGraphSchema;
 class UObject;
@@ -34,7 +35,6 @@ class BEHAVIORTREEEDITOR_API UBehaviorTreeGraphNode : public UAIGraphNode
 	TArray<TObjectPtr<UBehaviorTreeGraphNode>> Services;
 
 	//~ Begin UEdGraphNode Interface
-	virtual class UBehaviorTreeGraph* GetBehaviorTreeGraph();
 	virtual void AllocateDefaultPins() override;
 	virtual FText GetTooltipText() const override;
 	virtual bool CanCreateUnderSpecifiedSchema(const UEdGraphSchema* DesiredSchema) const override;
@@ -54,6 +54,12 @@ class BEHAVIORTREEEDITOR_API UBehaviorTreeGraphNode : public UAIGraphNode
 	virtual int32 FindSubNodeDropIndex(UAIGraphNode* SubNode) const override;
 	virtual void InsertSubNodeAt(UAIGraphNode* SubNode, int32 DropIndex) override;
 
+	UE_DEPRECATED(5.4, "Use GetOwnerBehaviorTreeGraph instead.")
+	virtual UBehaviorTreeGraph* GetBehaviorTreeGraph();
+	virtual UBehaviorTreeGraph* GetOwnerBehaviorTreeGraph() const;
+
+	virtual FLinearColor GetBackgroundColor(bool bIsActiveForDebugger) const;
+
 	/** check if node can accept breakpoints */
 	virtual bool CanPlaceBreakpoints() const { return false; }
 
diff --git a/Engine/Source/Editor/BehaviorTreeEditor/Classes/BehaviorTreeGraphNode_Composite.h b/Engine/Source/Editor/BehaviorTreeEditor/Classes/BehaviorTreeGraphNode_Composite.h
--- a/Engine/Source/Editor/BehaviorTreeEditor/Classes/BehaviorTreeGraphNode_Composite.h
+++ b/Engine/Source/Editor/BehaviorTreeEditor/Classes/BehaviorTreeGraphNode_Composite.h
@@ -29,6 +29,8 @@ class BEHAVIORTREEEDITOR_API UBehaviorTreeGraphNode_Composite : public UBehavior
 	/** check if node can accept breakpoints */
 	virtual bool CanPlaceBreakpoints() const override { return true; }
 
+	virtual FLinearColor GetBackgroundColor(bool bIsActiveForDebugger) const override;
+
 protected:
 
 	virtual void PostPasteNode() override;
diff --git a/Engine/Source/Editor/BehaviorTreeEditor/Classes/BehaviorTreeGraphNode_CompositeDecorator.h b/Engine/Source/Editor/BehaviorTreeEditor/Classes/BehaviorTreeGraphNode_CompositeDecorator.h
--- a/Engine/Source/Editor/BehaviorTreeEditor/Classes/BehaviorTreeGraphNode_CompositeDecorator.h
+++ b/Engine/Source/Editor/BehaviorTreeEditor/Classes/BehaviorTreeGraphNode_CompositeDecorator.h
@@ -71,6 +71,8 @@ class BEHAVIORTREEEDITOR_API UBehaviorTreeGraphNode_CompositeDecorator : public
 
 	virtual void PostEditChangeProperty(struct FPropertyChangedEvent& PropertyChangedEvent) override;
 
+	virtual FLinearColor GetBackgroundColor(bool bIsActiveForDebugger) const override;
+
 	void ResetExecutionRange();
 
 	/** Execution index range of internal nodes, used by debugger */
diff --git a/Engine/Source/Editor/BehaviorTreeEditor/Classes/BehaviorTreeGraphNode_Decorator.h b/Engine/Source/Editor/BehaviorTreeEditor/Classes/BehaviorTreeGraphNode_Decorator.h
--- a/Engine/Source/Editor/BehaviorTreeEditor/Classes/BehaviorTreeGraphNode_Decorator.h
+++ b/Engine/Source/Editor/BehaviorTreeEditor/Classes/BehaviorTreeGraphNode_Decorator.h
@@ -14,12 +14,13 @@
 
 class UObject;
 
-UCLASS()
+UCLASS(MinimalAPI)
 class UBehaviorTreeGraphNode_Decorator : public UBehaviorTreeGraphNode
 {
 	GENERATED_UCLASS_BODY()
 
-	virtual FText GetNodeTitle(ENodeTitleType::Type TitleType) const override;
-	virtual void AllocateDefaultPins() override;
-	void CollectDecoratorData(TArray<class UBTDecorator*>& NodeInstances, TArray<struct FBTDecoratorLogic>& Operations) const;
+	BEHAVIORTREEEDITOR_API virtual FText GetNodeTitle(ENodeTitleType::Type TitleType) const override;
+	BEHAVIORTREEEDITOR_API virtual void AllocateDefaultPins() override;
+	BEHAVIORTREEEDITOR_API virtual FLinearColor GetBackgroundColor(bool bIsActiveForDebugger) const override;
+	BEHAVIORTREEEDITOR_API void CollectDecoratorData(TArray<class UBTDecorator*>& NodeInstances, TArray<struct FBTDecoratorLogic>& Operations) const;
 };
diff --git a/Engine/Source/Editor/BehaviorTreeEditor/Classes/BehaviorTreeGraphNode_Root.h b/Engine/Source/Editor/BehaviorTreeEditor/Classes/BehaviorTreeGraphNode_Root.h
--- a/Engine/Source/Editor/BehaviorTreeEditor/Classes/BehaviorTreeGraphNode_Root.h
+++ b/Engine/Source/Editor/BehaviorTreeEditor/Classes/BehaviorTreeGraphNode_Root.h
@@ -39,6 +39,8 @@ class BEHAVIORTREEEDITOR_API UBehaviorTreeGraphNode_Root : public UBehaviorTreeG
 	virtual void PostEditUndo() override;
 	virtual FText GetDescription() const override;
 
+	virtual FLinearColor GetBackgroundColor(bool bIsActiveForDebugger) const override;
+
 	/** notify behavior tree about blackboard change */
 	void UpdateBlackboard();
 };
diff --git a/Engine/Source/Editor/BehaviorTreeEditor/Classes/BehaviorTreeGraphNode_Service.h b/Engine/Source/Editor/BehaviorTreeEditor/Classes/BehaviorTreeGraphNode_Service.h
--- a/Engine/Source/Editor/BehaviorTreeEditor/Classes/BehaviorTreeGraphNode_Service.h
+++ b/Engine/Source/Editor/BehaviorTreeEditor/Classes/BehaviorTreeGraphNode_Service.h
@@ -13,12 +13,13 @@
 
 class UObject;
 
-UCLASS()
+UCLASS(MinimalAPI)
 class UBehaviorTreeGraphNode_Service : public UBehaviorTreeGraphNode
 {
 	GENERATED_UCLASS_BODY()
 
-	virtual FText GetNodeTitle(ENodeTitleType::Type TitleType) const override;
-	virtual void AllocateDefaultPins() override;
+	BEHAVIORTREEEDITOR_API virtual FText GetNodeTitle(ENodeTitleType::Type TitleType) const override;
+	BEHAVIORTREEEDITOR_API virtual void AllocateDefaultPins() override;
+	BEHAVIORTREEEDITOR_API virtual FLinearColor GetBackgroundColor(bool bIsActiveForDebugger) const override;
 };
 
diff --git a/Engine/Source/Editor/BehaviorTreeEditor/Classes/BehaviorTreeGraphNode_SubtreeTask.h b/Engine/Source/Editor/BehaviorTreeEditor/Classes/BehaviorTreeGraphNode_SubtreeTask.h
--- a/Engine/Source/Editor/BehaviorTreeEditor/Classes/BehaviorTreeGraphNode_SubtreeTask.h
+++ b/Engine/Source/Editor/BehaviorTreeEditor/Classes/BehaviorTreeGraphNode_SubtreeTask.h
@@ -26,4 +26,6 @@ class UBehaviorTreeGraphNode_SubtreeTask : public UBehaviorTreeGraphNode_Task
 
 	/** updates nodes injected from subtree's root */
 	bool UpdateInjectedNodes();
+
+	virtual FLinearColor GetBackgroundColor(bool bIsActiveForDebugger) const override;
 };
diff --git a/Engine/Source/Editor/BehaviorTreeEditor/Classes/BehaviorTreeGraphNode_Task.h b/Engine/Source/Editor/BehaviorTreeEditor/Classes/BehaviorTreeGraphNode_Task.h
--- a/Engine/Source/Editor/BehaviorTreeEditor/Classes/BehaviorTreeGraphNode_Task.h
+++ b/Engine/Source/Editor/BehaviorTreeEditor/Classes/BehaviorTreeGraphNode_Task.h
@@ -13,15 +13,16 @@
 
 class UObject;
 
-UCLASS()
+UCLASS(MinimalAPI)
 class UBehaviorTreeGraphNode_Task : public UBehaviorTreeGraphNode
 {
 	GENERATED_UCLASS_BODY()
 
-	virtual void AllocateDefaultPins() override;
-	virtual FText GetNodeTitle(ENodeTitleType::Type TitleType) const override;
+	BEHAVIORTREEEDITOR_API virtual void AllocateDefaultPins() override;
+	BEHAVIORTREEEDITOR_API virtual FText GetNodeTitle(ENodeTitleType::Type TitleType) const override;
 	/** Gets a list of actions that can be done to this particular node */
-	virtual void GetNodeContextMenuActions(class UToolMenu* Menu, class UGraphNodeContextMenuContext* Context) const override;
+	BEHAVIORTREEEDITOR_API virtual void GetNodeContextMenuActions(class UToolMenu* Menu, class UGraphNodeContextMenuContext* Context) const override;
 
-	virtual bool CanPlaceBreakpoints() const override { return true; }
+	BEHAVIORTREEEDITOR_API virtual bool CanPlaceBreakpoints() const override { return true; }
+	BEHAVIORTREEEDITOR_API virtual FLinearColor GetBackgroundColor(bool bIsActiveForDebugger) const override;
 };
diff --git a/Engine/Source/Editor/BehaviorTreeEditor/Classes/EdGraphSchema_BehaviorTree.h b/Engine/Source/Editor/BehaviorTreeEditor/Classes/EdGraphSchema_BehaviorTree.h
--- a/Engine/Source/Editor/BehaviorTreeEditor/Classes/EdGraphSchema_BehaviorTree.h
+++ b/Engine/Source/Editor/BehaviorTreeEditor/Classes/EdGraphSchema_BehaviorTree.h
@@ -20,6 +20,9 @@
 
 class FSlateRect;
 class UBehaviorTreeGraphNode_CompositeDecorator;
+class UBehaviorTreeGraphNode_Decorator;
+class UBehaviorTreeGraphNode_Service;
+class UBehaviorTreeGraphNode_Task;
 class UClass;
 class UEdGraph;
 class UEdGraphNode;
@@ -68,6 +71,9 @@ class BEHAVIORTREEEDITOR_API UEdGraphSchema_BehaviorTree : public UAIGraphSchema
 	virtual void GetSubNodeClasses(int32 SubNodeFlags, TArray<FGraphNodeClassData>& ClassData, UClass*& GraphNodeClass) const override;
 
 	TSubclassOf<UBehaviorTreeGraphNode_CompositeDecorator> CompositeDecoratorClass;
+	TSubclassOf<UBehaviorTreeGraphNode_Decorator> DecoratorClass;
+	TSubclassOf<UBehaviorTreeGraphNode_Task> TaskClass;
+	TSubclassOf<UBehaviorTreeGraphNode_Service> ServiceClass;
 
 protected:
 	virtual FGraphNodeClassHelper& GetClassCache() const;
diff --git a/Engine/Source/Editor/BehaviorTreeEditor/Private/BehaviorTreeColors.h b/Engine/Source/Editor/BehaviorTreeEditor/Private/BehaviorTreeColors.h
deleted file mode 100644
diff --git a/Engine/Source/Editor/BehaviorTreeEditor/Private/BehaviorTreeEditor.cpp b/Engine/Source/Editor/BehaviorTreeEditor/Private/BehaviorTreeEditor.cpp
--- a/Engine/Source/Editor/BehaviorTreeEditor/Private/BehaviorTreeEditor.cpp
+++ b/Engine/Source/Editor/BehaviorTreeEditor/Private/BehaviorTreeEditor.cpp
@@ -119,6 +119,9 @@ class SWidget;
 const FName FBehaviorTreeEditor::BehaviorTreeMode(TEXT("BehaviorTree"));
 const FName FBehaviorTreeEditor::BlackboardMode(TEXT("Blackboard"));
 
+FText FBehaviorTreeEditor::BehaviorTreeModeText(LOCTEXT("BehaviorTreeMode", "Behavior Tree"));
+FText FBehaviorTreeEditor::BlackboardModeText(LOCTEXT("BlackboardMode", "Blackboard"));
+
 //////////////////////////////////////////////////////////////////////////
 FBehaviorTreeEditor::FBehaviorTreeEditor() 
 	: IBehaviorTreeEditor()
@@ -144,6 +147,7 @@ FBehaviorTreeEditor::FBehaviorTreeEditor()
 	GraphName = "Behavior Tree";
 	CornerText = LOCTEXT("AppearanceCornerText", "BEHAVIOR TREE");
 	TitleText = LOCTEXT("BehaviorTreeGraphLabel", "Behavior Tree");
+	RootNodeNoteText = LOCTEXT("RootLevelNode", "Root-level decorators are only valid and will be executed if this BT is be used\nas static a sub-tree (via \"Run Behavior\"). These decorators will be ignored if\ndynamically injected with \"Run Dynamic Behavior\".");
 }
 
 FBehaviorTreeEditor::~FBehaviorTreeEditor()
@@ -696,7 +700,7 @@ TSharedRef<SWidget> FBehaviorTreeEditor::SpawnProperties()
 				.Padding(FMargin(5.0f))
 				[
 					SNew(STextBlock)
-					.Text(LOCTEXT("RootLevelNode", "Root-level decorators are only valid and will be executed if this BT is be used\nas static a sub-tree (via \"Run Behavior\"). These decorators will be ignored if\ndynamically injected with \"Run Dynamic Behavior\"."))
+					.Text(RootNodeNoteText)
 				]
 			]
 			+SVerticalBox::Slot()
@@ -1589,8 +1593,8 @@ FText FBehaviorTreeEditor::GetLocalizedMode(FName InMode)
 
 	if (LocModes.Num() == 0)
 	{
-		LocModes.Add( BehaviorTreeMode, LOCTEXT("BehaviorTreeMode", "Behavior Tree") );
-		LocModes.Add( BlackboardMode, LOCTEXT("BlackboardMode", "Blackboard") );
+		LocModes.Add( BehaviorTreeMode, BehaviorTreeModeText );
+		LocModes.Add( BlackboardMode, BlackboardModeText );
 	}
 
 	check( InMode != NAME_None );
diff --git a/Engine/Source/Editor/BehaviorTreeEditor/Private/BehaviorTreeGraphNode.cpp b/Engine/Source/Editor/BehaviorTreeEditor/Private/BehaviorTreeGraphNode.cpp
--- a/Engine/Source/Editor/BehaviorTreeEditor/Private/BehaviorTreeGraphNode.cpp
+++ b/Engine/Source/Editor/BehaviorTreeEditor/Private/BehaviorTreeGraphNode.cpp
@@ -1,6 +1,7 @@
 // Copyright Epic Games, Inc. All Rights Reserved.
 
 #include "BehaviorTreeGraphNode.h"
+#include "BehaviorTreeColors.h"
 #include "Widgets/DeclarativeSyntaxSupport.h"
 #include "ToolMenus.h"
 #include "BehaviorTree/BTNode.h"
@@ -105,6 +106,11 @@ FText UBehaviorTreeGraphNode::GetTooltipText() const
 }
 
 UBehaviorTreeGraph* UBehaviorTreeGraphNode::GetBehaviorTreeGraph()
+{
+	return GetOwnerBehaviorTreeGraph();
+}
+
+UBehaviorTreeGraph* UBehaviorTreeGraphNode::GetOwnerBehaviorTreeGraph() const
 {
 	return CastChecked<UBehaviorTreeGraph>(GetGraph());
 }
@@ -327,6 +333,11 @@ void UBehaviorTreeGraphNode::InsertSubNodeAt(UAIGraphNode* SubNode, int32 DropIn
 	}
 }
 
+FLinearColor UBehaviorTreeGraphNode::GetBackgroundColor(bool bIsActiveForDebugger) const
+{
+	return BehaviorTreeColors::NodeBody::Default;
+}
+
 void UBehaviorTreeGraphNode::CreateAddDecoratorSubMenu(UToolMenu* Menu, UEdGraph* Graph) const
 {
 	TSharedRef<SGraphEditorActionMenuAI> Widget =
diff --git a/Engine/Source/Editor/BehaviorTreeEditor/Private/BehaviorTreeGraphNode_Composite.cpp b/Engine/Source/Editor/BehaviorTreeEditor/Private/BehaviorTreeGraphNode_Composite.cpp
--- a/Engine/Source/Editor/BehaviorTreeEditor/Private/BehaviorTreeGraphNode_Composite.cpp
+++ b/Engine/Source/Editor/BehaviorTreeEditor/Private/BehaviorTreeGraphNode_Composite.cpp
@@ -2,6 +2,8 @@
 
 #include "BehaviorTreeGraphNode_Composite.h"
 
+#include "BehaviorTreeColors.h"
+#include "BehaviorTreeGraph.h"
 #include "BehaviorTree/BTCompositeNode.h"
 #include "BehaviorTree/BTNode.h"
 #include "HAL/PlatformMath.h"
@@ -51,6 +53,13 @@ FText UBehaviorTreeGraphNode_Composite::GetTooltipText() const
 	return Super::GetTooltipText();
 }
 
+FLinearColor UBehaviorTreeGraphNode_Composite::GetBackgroundColor(bool bIsActiveForDebugger) const
+{
+	UBTCompositeNode* CompositeNodeInstance = Cast<UBTCompositeNode>(NodeInstance);
+	const bool bIsScoped = CompositeNodeInstance && CompositeNodeInstance->IsApplyingDecoratorScope();
+	return bIsScoped ? BehaviorTreeColors::NodeBody::CompositeScoped : BehaviorTreeColors::NodeBody::Composite;
+}
+
 void UBehaviorTreeGraphNode_Composite::PostPasteNode()
 {
 	Super::PostPasteNode();
@@ -67,5 +76,9 @@ void UBehaviorTreeGraphNode_Composite::PostPasteNode()
 void UBehaviorTreeGraphNode_Composite::GetNodeContextMenuActions(UToolMenu* Menu, UGraphNodeContextMenuContext* Context) const
 {
 	AddContextMenuActionsDecorators(Menu, "BehaviorTreeGraphNode", Context);
-	AddContextMenuActionsServices(Menu, "BehaviorTreeGraphNode", Context);
+
+	if (GetOwnerBehaviorTreeGraph()->DoesSupportServices())
+	{
+		AddContextMenuActionsServices(Menu, "BehaviorTreeGraphNode", Context);
+	}
 }
diff --git a/Engine/Source/Editor/BehaviorTreeEditor/Private/BehaviorTreeGraphNode_CompositeDecorator.cpp b/Engine/Source/Editor/BehaviorTreeEditor/Private/BehaviorTreeGraphNode_CompositeDecorator.cpp
--- a/Engine/Source/Editor/BehaviorTreeEditor/Private/BehaviorTreeGraphNode_CompositeDecorator.cpp
+++ b/Engine/Source/Editor/BehaviorTreeEditor/Private/BehaviorTreeGraphNode_CompositeDecorator.cpp
@@ -2,6 +2,7 @@
 
 #include "BehaviorTreeGraphNode_CompositeDecorator.h"
 
+#include "BehaviorTreeColors.h"
 #include "BehaviorTree/BTCompositeNode.h"
 #include "BehaviorTree/BTDecorator.h"
 #include "BehaviorTree/BTNode.h"
@@ -306,6 +307,15 @@ void UBehaviorTreeGraphNode_CompositeDecorator::PostEditChangeProperty(struct FP
 	}
 }
 
+FLinearColor UBehaviorTreeGraphNode_CompositeDecorator::GetBackgroundColor(bool bIsActiveForDebugger) const
+{
+	return bIsActiveForDebugger
+		? BehaviorTreeColors::Debugger::ActiveDecorator
+		: (bInjectedNode || bRootLevel)
+			? BehaviorTreeColors::NodeBody::InjectedSubNode
+			: BehaviorTreeColors::NodeBody::Decorator;
+}
+
 struct FLogicDesc
 {
 	FString OperationDesc;
diff --git a/Engine/Source/Editor/BehaviorTreeEditor/Private/BehaviorTreeGraphNode_Decorator.cpp b/Engine/Source/Editor/BehaviorTreeEditor/Private/BehaviorTreeGraphNode_Decorator.cpp
--- a/Engine/Source/Editor/BehaviorTreeEditor/Private/BehaviorTreeGraphNode_Decorator.cpp
+++ b/Engine/Source/Editor/BehaviorTreeEditor/Private/BehaviorTreeGraphNode_Decorator.cpp
@@ -3,6 +3,7 @@
 #include "BehaviorTreeGraphNode_Decorator.h"
 
 #include "AIGraphTypes.h"
+#include "BehaviorTreeColors.h"
 #include "BehaviorTree/BTCompositeNode.h"
 #include "BehaviorTree/BTDecorator.h"
 #include "Containers/UnrealString.h"
@@ -23,6 +24,15 @@ void UBehaviorTreeGraphNode_Decorator::AllocateDefaultPins()
 	//No Pins for decorators
 }
 
+FLinearColor UBehaviorTreeGraphNode_Decorator::GetBackgroundColor(bool bIsActiveForDebugger) const
+{
+	return bIsActiveForDebugger
+		? BehaviorTreeColors::Debugger::ActiveDecorator
+		: (bInjectedNode || bRootLevel)
+			? BehaviorTreeColors::NodeBody::InjectedSubNode
+			: BehaviorTreeColors::NodeBody::Decorator;
+}
+
 FText UBehaviorTreeGraphNode_Decorator::GetNodeTitle(ENodeTitleType::Type TitleType) const
 {
 	const UBTDecorator* Decorator = Cast<UBTDecorator>(NodeInstance);
diff --git a/Engine/Source/Editor/BehaviorTreeEditor/Private/BehaviorTreeGraphNode_Root.cpp b/Engine/Source/Editor/BehaviorTreeEditor/Private/BehaviorTreeGraphNode_Root.cpp
--- a/Engine/Source/Editor/BehaviorTreeEditor/Private/BehaviorTreeGraphNode_Root.cpp
+++ b/Engine/Source/Editor/BehaviorTreeEditor/Private/BehaviorTreeGraphNode_Root.cpp
@@ -3,6 +3,7 @@
 #include "BehaviorTreeGraphNode_Root.h"
 
 #include "AISystem.h"
+#include "BehaviorTreeColors.h"
 #include "BehaviorTree/BehaviorTree.h"
 #include "BehaviorTree/BlackboardData.h"
 #include "BehaviorTreeEditorTypes.h"
@@ -90,9 +91,19 @@ FText UBehaviorTreeGraphNode_Root::GetDescription() const
 	return FText::FromString(GetNameSafe(BlackboardAsset));
 }
 
+FLinearColor UBehaviorTreeGraphNode_Root::GetBackgroundColor(bool bIsActiveForDebugger) const
+{
+	if (Pins.IsValidIndex(0) && Pins[0]->LinkedTo.Num() > 0)
+	{
+		return BehaviorTreeColors::NodeBody::Root;
+	}
+
+	return Super::GetBackgroundColor(bIsActiveForDebugger);
+}
+
 void UBehaviorTreeGraphNode_Root::UpdateBlackboard()
 {
-	UBehaviorTreeGraph* MyGraph = GetBehaviorTreeGraph();
+	UBehaviorTreeGraph* MyGraph = GetOwnerBehaviorTreeGraph();
 	UBehaviorTree* BTAsset = Cast<UBehaviorTree>(MyGraph->GetOuter());
 	if (BTAsset && BTAsset->BlackboardAsset != BlackboardAsset)
 	{
diff --git a/Engine/Source/Editor/BehaviorTreeEditor/Private/BehaviorTreeGraphNode_Service.cpp b/Engine/Source/Editor/BehaviorTreeEditor/Private/BehaviorTreeGraphNode_Service.cpp
--- a/Engine/Source/Editor/BehaviorTreeEditor/Private/BehaviorTreeGraphNode_Service.cpp
+++ b/Engine/Source/Editor/BehaviorTreeEditor/Private/BehaviorTreeGraphNode_Service.cpp
@@ -3,6 +3,7 @@
 #include "BehaviorTreeGraphNode_Service.h"
 
 #include "AIGraphTypes.h"
+#include "BehaviorTreeColors.h"
 #include "BehaviorTree/BTService.h"
 #include "Containers/UnrealString.h"
 #include "HAL/PlatformMath.h"
@@ -20,6 +21,11 @@ void UBehaviorTreeGraphNode_Service::AllocateDefaultPins()
 	//No Pins for services
 }
 
+FLinearColor UBehaviorTreeGraphNode_Service::GetBackgroundColor(bool bIsActiveForDebugger) const
+{
+	return bIsActiveForDebugger ? BehaviorTreeColors::Debugger::ActiveService : BehaviorTreeColors::NodeBody::Service;
+}
+
 
 FText UBehaviorTreeGraphNode_Service::GetNodeTitle(ENodeTitleType::Type TitleType) const
 {
diff --git a/Engine/Source/Editor/BehaviorTreeEditor/Private/BehaviorTreeGraphNode_SubtreeTask.cpp b/Engine/Source/Editor/BehaviorTreeEditor/Private/BehaviorTreeGraphNode_SubtreeTask.cpp
--- a/Engine/Source/Editor/BehaviorTreeEditor/Private/BehaviorTreeGraphNode_SubtreeTask.cpp
+++ b/Engine/Source/Editor/BehaviorTreeEditor/Private/BehaviorTreeGraphNode_SubtreeTask.cpp
@@ -2,6 +2,7 @@
 
 #include "BehaviorTreeGraphNode_SubtreeTask.h"
 
+#include "BehaviorTreeColors.h"
 #include "BehaviorTree/BTDecorator.h"
 #include "BehaviorTree/BehaviorTree.h"
 #include "BehaviorTree/Tasks/BTTask_RunBehavior.h"
@@ -75,7 +76,7 @@ bool UBehaviorTreeGraphNode_SubtreeTask::UpdateInjectedNodes()
 	// add root level subnodes as injected nodes
 	if (SubRoot)
 	{
-		UBehaviorTree* BTAsset = Cast<UBehaviorTree>(GetBehaviorTreeGraph()->GetOuter());
+		UBehaviorTree* BTAsset = Cast<UBehaviorTree>(GetOwnerBehaviorTreeGraph()->GetOuter());
 
 		if(BTAsset)
 		{
@@ -139,3 +140,8 @@ bool UBehaviorTreeGraphNode_SubtreeTask::UpdateInjectedNodes()
 
 	return bUpdated;
 }
+
+FLinearColor UBehaviorTreeGraphNode_SubtreeTask::GetBackgroundColor(bool bIsActiveForDebugger) const
+{
+	return BehaviorTreeColors::NodeBody::TaskSpecial;
+}
diff --git a/Engine/Source/Editor/BehaviorTreeEditor/Private/BehaviorTreeGraphNode_Task.cpp b/Engine/Source/Editor/BehaviorTreeEditor/Private/BehaviorTreeGraphNode_Task.cpp
--- a/Engine/Source/Editor/BehaviorTreeEditor/Private/BehaviorTreeGraphNode_Task.cpp
+++ b/Engine/Source/Editor/BehaviorTreeEditor/Private/BehaviorTreeGraphNode_Task.cpp
@@ -3,8 +3,10 @@
 #include "BehaviorTreeGraphNode_Task.h"
 
 #include "AIGraphTypes.h"
+#include "BehaviorTreeColors.h"
 #include "BehaviorTree/BTNode.h"
 #include "BehaviorTreeEditorTypes.h"
+#include "BehaviorTreeGraph.h"
 #include "Containers/UnrealString.h"
 #include "HAL/PlatformMath.h"
 #include "Internationalization/Internationalization.h"
@@ -42,5 +44,14 @@ FText UBehaviorTreeGraphNode_Task::GetNodeTitle(ENodeTitleType::Type TitleType)
 void UBehaviorTreeGraphNode_Task::GetNodeContextMenuActions(UToolMenu* Menu, UGraphNodeContextMenuContext* Context) const
 {
 	AddContextMenuActionsDecorators(Menu, "BehaviorTreeGraphNode", Context);
-	AddContextMenuActionsServices(Menu, "BehaviorTreeGraphNode", Context);
+
+	if (GetOwnerBehaviorTreeGraph()->DoesSupportServices())
+	{
+		AddContextMenuActionsServices(Menu, "BehaviorTreeGraphNode", Context);
+	}
+}
+
+FLinearColor UBehaviorTreeGraphNode_Task::GetBackgroundColor(bool bIsActiveForDebugger) const
+{
+	return BehaviorTreeColors::NodeBody::Task;
 }
diff --git a/Engine/Source/Editor/BehaviorTreeEditor/Private/EdGraphSchema_BehaviorTree.cpp b/Engine/Source/Editor/BehaviorTreeEditor/Private/EdGraphSchema_BehaviorTree.cpp
--- a/Engine/Source/Editor/BehaviorTreeEditor/Private/EdGraphSchema_BehaviorTree.cpp
+++ b/Engine/Source/Editor/BehaviorTreeEditor/Private/EdGraphSchema_BehaviorTree.cpp
@@ -67,6 +67,9 @@ UEdGraphNode* FBehaviorTreeSchemaAction_AutoArrange::PerformAction(class UEdGrap
 UEdGraphSchema_BehaviorTree::UEdGraphSchema_BehaviorTree(const FObjectInitializer& ObjectInitializer) : Super(ObjectInitializer)
 {
 	CompositeDecoratorClass = UBehaviorTreeGraphNode_CompositeDecorator::StaticClass();
+	DecoratorClass = UBehaviorTreeGraphNode_Decorator::StaticClass();
+	TaskClass = UBehaviorTreeGraphNode_Task::StaticClass();
+	ServiceClass = UBehaviorTreeGraphNode_Service::StaticClass();
 }
 
 void UEdGraphSchema_BehaviorTree::CreateDefaultNodesForGraph(UEdGraph& Graph) const
@@ -99,12 +102,12 @@ void UEdGraphSchema_BehaviorTree::GetSubNodeClasses(int32 SubNodeFlags, TArray<F
 	if (SubNodeFlags == ESubNode::Decorator)
 	{
 		ClassCache.GatherClasses(UBTDecorator::StaticClass(), ClassData);
-		GraphNodeClass = UBehaviorTreeGraphNode_Decorator::StaticClass();
+		GraphNodeClass = DecoratorClass;
 	}
 	else
 	{
 		ClassCache.GatherClasses(UBTService::StaticClass(), ClassData);
-		GraphNodeClass = UBehaviorTreeGraphNode_Service::StaticClass();
+		GraphNodeClass = ServiceClass;
 	}
 }
 
@@ -176,7 +179,7 @@ void UEdGraphSchema_BehaviorTree::GetGraphContextActions(FGraphContextMenuBuilde
 
 			TSharedPtr<FAISchemaAction_NewNode> AddOpAction = UAIGraphSchema::AddNewNodeAction(TasksBuilder, NodeClass.GetCategory(), NodeTypeName, NodeClass.GetTooltip());
 			
-			UClass* GraphNodeClass = UBehaviorTreeGraphNode_Task::StaticClass();
+			UClass* GraphNodeClass = TaskClass;
 			if (IsNodeSubtreeTask(NodeClass))
 			{
 				GraphNodeClass = UBehaviorTreeGraphNode_SubtreeTask::StaticClass();
diff --git a/Engine/Source/Editor/BehaviorTreeEditor/Private/SGraphNode_BehaviorTree.cpp b/Engine/Source/Editor/BehaviorTreeEditor/Private/SGraphNode_BehaviorTree.cpp
--- a/Engine/Source/Editor/BehaviorTreeEditor/Private/SGraphNode_BehaviorTree.cpp
+++ b/Engine/Source/Editor/BehaviorTreeEditor/Private/SGraphNode_BehaviorTree.cpp
@@ -241,7 +241,6 @@ FSlateColor SGraphNode_BehaviorTree::GetBorderBackgroundColor() const
 FSlateColor SGraphNode_BehaviorTree::GetBackgroundColor() const
 {
 	UBehaviorTreeGraphNode* BTGraphNode = Cast<UBehaviorTreeGraphNode>(GraphNode);
-	UBehaviorTreeGraphNode_Decorator* BTGraph_Decorator = Cast<UBehaviorTreeGraphNode_Decorator>(GraphNode);
 	const bool bIsActiveForDebugger = BTGraphNode ?
 		!bSuppressDebuggerColor && (BTGraphNode->bDebuggerMarkCurrentlyActive || BTGraphNode->bDebuggerMarkPreviouslyActive) :
 		false;
@@ -251,36 +250,9 @@ FSlateColor SGraphNode_BehaviorTree::GetBackgroundColor() const
 	{
 		NodeColor = BehaviorTreeColors::NodeBody::Error;
 	}
-	else if (BTGraphNode && BTGraphNode->bInjectedNode)
+	else if (BTGraphNode)
 	{
-		NodeColor = bIsActiveForDebugger ? BehaviorTreeColors::Debugger::ActiveDecorator : BehaviorTreeColors::NodeBody::InjectedSubNode;
-	}
-	else if (BTGraph_Decorator || Cast<UBehaviorTreeGraphNode_CompositeDecorator>(GraphNode))
-	{
-		check(BTGraphNode);
-		NodeColor = bIsActiveForDebugger ? BehaviorTreeColors::Debugger::ActiveDecorator : 
-			BTGraphNode->bRootLevel ? BehaviorTreeColors::NodeBody::InjectedSubNode : BehaviorTreeColors::NodeBody::Decorator;
-	}
-	else if (Cast<UBehaviorTreeGraphNode_Task>(GraphNode))
-	{
-		check(BTGraphNode);
-		const bool bIsSpecialTask = Cast<UBTTask_RunBehavior>(BTGraphNode->NodeInstance) != nullptr;
-		NodeColor = bIsSpecialTask ? BehaviorTreeColors::NodeBody::TaskSpecial : BehaviorTreeColors::NodeBody::Task;
-	}
-	else if (Cast<UBehaviorTreeGraphNode_Composite>(GraphNode))
-	{
-		check(BTGraphNode);
-		UBTCompositeNode* CompositeNodeInstance = Cast<UBTCompositeNode>(BTGraphNode->NodeInstance);
-		const bool bIsScoped = CompositeNodeInstance && CompositeNodeInstance->IsApplyingDecoratorScope();
-		NodeColor = bIsScoped ? BehaviorTreeColors::NodeBody::CompositeScoped : BehaviorTreeColors::NodeBody::Composite;
-	}
-	else if (Cast<UBehaviorTreeGraphNode_Service>(GraphNode))
-	{
-		NodeColor = bIsActiveForDebugger ? BehaviorTreeColors::Debugger::ActiveService : BehaviorTreeColors::NodeBody::Service;
-	}
-	else if (Cast<UBehaviorTreeGraphNode_Root>(GraphNode) && GraphNode->Pins.IsValidIndex(0) && GraphNode->Pins[0]->LinkedTo.Num() > 0)
-	{
-		NodeColor = BehaviorTreeColors::NodeBody::Root;
+		NodeColor = BTGraphNode->GetBackgroundColor(bIsActiveForDebugger);
 	}
 
 	return (FlashAlpha > 0.0f) ? FMath::Lerp(NodeColor, FlashColor, FlashAlpha) : NodeColor;
@@ -1133,7 +1105,7 @@ void SGraphNode_BehaviorTree::MoveTo(const FVector2D& NewPosition, FNodeSet& Nod
 	UBehaviorTreeGraphNode* BTGraphNode = Cast<UBehaviorTreeGraphNode>(GraphNode);
 	if (BTGraphNode && !BTGraphNode->IsSubNode())
 	{
-		UBehaviorTreeGraph* BTGraph = BTGraphNode->GetBehaviorTreeGraph();
+		UBehaviorTreeGraph* BTGraph = BTGraphNode->GetOwnerBehaviorTreeGraph();
 		if (BTGraph)
 		{
 			for (int32 Idx = 0; Idx < BTGraphNode->Pins.Num(); Idx++)
diff --git a/Engine/Source/Editor/BehaviorTreeEditor/Public/BehaviorTreeColors.h b/Engine/Source/Editor/BehaviorTreeEditor/Public/BehaviorTreeColors.h
new file mode 100644
--- /dev/null
+++ b/Engine/Source/Editor/BehaviorTreeEditor/Public/BehaviorTreeColors.h
@@ -0,0 +1,70 @@
+// Copyright Epic Games, Inc. All Rights Reserved.
+
+#pragma once
+
+#include "CoreMinimal.h"
+
+class Error;
+
+namespace BehaviorTreeColors
+{
+	namespace NodeBody
+	{
+		const FLinearColor Task(0.24f, 0.055f, 0.715f);
+		const FLinearColor TaskSpecial(0.1f, 0.05f, 0.2f);
+		const FLinearColor Decorator(0.0f, 0.07f, 0.4f);	
+		const FLinearColor InjectedSubNode(0.0f, 0.02f, 0.1f);
+		const FLinearColor Composite(0.1f, 0.1f, 0.1f);
+		const FLinearColor CompositeScoped(0.05f, 0.05f, 0.05f);
+		const FLinearColor Service(0.0f, 0.4f, 0.22f);
+		const FLinearColor Default(0.15f, 0.15f, 0.15f);
+		const FLinearColor Root(0.5f, 0.5f, 0.5f, 0.1f);
+		const FLinearColor Error(1.0f, 0.0f, 0.0f);
+	}
+
+	namespace NodeBorder
+	{
+		const FLinearColor Inactive(0.08f, 0.08f, 0.08f);
+		const FLinearColor Root(0.2f, 0.2f, 0.2f, 0.2f);
+		const FLinearColor Selected(1.00f, 0.08f, 0.08f);
+		const FLinearColor ActiveDebugging(1.0f, 1.0f, 0.0f);
+		const FLinearColor InactiveDebugging(0.4f, 0.4f, 0.0f);
+		const FLinearColor HighlightAbortRange0(0.0f, 0.22f, 0.4f);
+		const FLinearColor HighlightAbortRange1(0.0f, 0.4f, 0.22f);
+		const FLinearColor Disconnected(0.f, 0.f, 0.f);
+		const FLinearColor BrokenWithParent(1.f, 0.f, 1.f);
+		const FLinearColor QuickFind(0.f, 0.8f, 0.f);
+	}
+
+	namespace Pin
+	{
+		const FLinearColor Diff(0.9f, 0.2f, 0.15f);
+		const FLinearColor Hover(1.0f, 0.7f, 0.0f);
+		const FLinearColor Default(0.02f, 0.02f, 0.02f);
+		const FLinearColor TaskOnly(0.13f, 0.03f, 0.4f);
+		const FLinearColor CompositeOnly(0.02f, 0.02f, 0.02f);
+		const FLinearColor SingleNode(0.02f, 0.02f, 0.02f);
+	}
+	
+	namespace Connection
+	{
+		const FLinearColor Default(1.0f, 1.0f, 1.0f);
+	}
+
+	namespace Debugger
+	{
+		const FLinearColor ActiveDecorator(0.05f, 0.27f, 0.7f);
+		const FLinearColor ActiveService(0.3f, 0.8f, 0.4f);
+		const FLinearColor TaskFlash(1.0f, 1.0f, 0.0f);
+		const FLinearColor SearchSucceeded(0.0f, 1.0f, 0.0f);
+		const FLinearColor SearchFailed(1.0f, 0.0f, 0.0f);
+
+		const FLinearColor DescHeader(0.3f, 0.8f, 0.4f);
+		const FLinearColor DescKeys(0.3f, 0.4f, 0.8f);
+	}
+
+	namespace Action
+	{
+		const FLinearColor DragMarker(1.0f, 1.0f, 0.2f);
+	}
+}
diff --git a/Engine/Source/Editor/BehaviorTreeEditor/Public/BehaviorTreeEditor.h b/Engine/Source/Editor/BehaviorTreeEditor/Public/BehaviorTreeEditor.h
--- a/Engine/Source/Editor/BehaviorTreeEditor/Public/BehaviorTreeEditor.h
+++ b/Engine/Source/Editor/BehaviorTreeEditor/Public/BehaviorTreeEditor.h
@@ -265,6 +265,7 @@ protected:
 	FName GraphName;
 	FText CornerText;
 	FText TitleText;
+	FText RootNodeNoteText;
 
 private:
 	/** Create widget for graph editing */
@@ -350,4 +351,7 @@ public:
 	/** Modes in mode switcher */
 	static const FName BehaviorTreeMode;
 	static const FName BlackboardMode;
+
+	static FText BehaviorTreeModeText;
+	static FText BlackboardModeText;
 };
