[Behavior Tree] Add support for node-specific error reporting

diff --git a/Engine/Source/Editor/AIGraph/Classes/AIGraphNode.h b/Engine/Source/Editor/AIGraph/Classes/AIGraphNode.h
--- a/Engine/Source/Editor/AIGraph/Classes/AIGraphNode.h
+++ b/Engine/Source/Editor/AIGraph/Classes/AIGraphNode.h
@@ -100,6 +100,12 @@ class AIGRAPH_API UAIGraphNode : public UEdGraphNode
 	/** updates ClassData from node instance */
 	virtual void UpdateNodeClassData();
 
+	/**
+	 * Checks for any errors in this node and updates ErrorMessage with any resulting message
+	 * Called every time the graph is serialized (i.e. loaded, saved, execution index changed, etc)
+	 */
+	virtual void UpdateErrorMessage();
+
 	/** Check if node instance uses blueprint for its implementation */
 	bool UsesBlueprint() const;
 
diff --git a/Engine/Source/Editor/AIGraph/Private/AIGraph.cpp b/Engine/Source/Editor/AIGraph/Private/AIGraph.cpp
--- a/Engine/Source/Editor/AIGraph/Private/AIGraph.cpp
+++ b/Engine/Source/Editor/AIGraph/Private/AIGraph.cpp
@@ -93,6 +93,15 @@ void UpdateAIGraphNodeErrorMessage(UAIGraphNode& Node)
 	if (Node.NodeInstance)
 	{
 		Node.ErrorMessage = FGraphNodeClassHelper::GetDeprecationMessage(Node.NodeInstance->GetClass());
+
+		// Only check for node-specific errors if the node is not deprecated
+		if (Node.ErrorMessage.IsEmpty())
+		{
+			Node.UpdateErrorMessage();
+
+			// For node-specific validation we don't want to spam the log with errors
+			return;
+		}
 	}
 	else
 	{
diff --git a/Engine/Source/Editor/AIGraph/Private/AIGraphNode.cpp b/Engine/Source/Editor/AIGraph/Private/AIGraphNode.cpp
--- a/Engine/Source/Editor/AIGraph/Private/AIGraphNode.cpp
+++ b/Engine/Source/Editor/AIGraph/Private/AIGraphNode.cpp
@@ -374,10 +374,15 @@ void UAIGraphNode::UpdateNodeClassData()
 	if (NodeInstance)
 	{
 		UpdateNodeClassDataFrom(NodeInstance->GetClass(), ClassData);
-		ErrorMessage = ClassData.GetDeprecatedMessage();
+		UpdateErrorMessage();
 	}
 }
 
+void UAIGraphNode::UpdateErrorMessage()
+{
+	ErrorMessage = ClassData.GetDeprecatedMessage();
+}
+
 void UAIGraphNode::UpdateNodeClassDataFrom(UClass* InstanceClass, FGraphNodeClassData& UpdatedData)
 {
 	if (InstanceClass)
diff --git a/Engine/Source/Editor/BehaviorTreeEditor/Classes/BehaviorTreeGraphNode.h b/Engine/Source/Editor/BehaviorTreeEditor/Classes/BehaviorTreeGraphNode.h
--- a/Engine/Source/Editor/BehaviorTreeEditor/Classes/BehaviorTreeGraphNode.h
+++ b/Engine/Source/Editor/BehaviorTreeEditor/Classes/BehaviorTreeGraphNode.h
@@ -53,6 +53,7 @@ class BEHAVIORTREEEDITOR_API UBehaviorTreeGraphNode : public UAIGraphNode
 	virtual void RemoveAllSubNodes() override;
 	virtual int32 FindSubNodeDropIndex(UAIGraphNode* SubNode) const override;
 	virtual void InsertSubNodeAt(UAIGraphNode* SubNode, int32 DropIndex) override;
+	virtual void UpdateErrorMessage() override;
 
 	UE_DEPRECATED(5.4, "Use GetOwnerBehaviorTreeGraph instead.")
 	virtual UBehaviorTreeGraph* GetBehaviorTreeGraph();
diff --git a/Engine/Source/Editor/BehaviorTreeEditor/Private/BehaviorTreeGraphNode.cpp b/Engine/Source/Editor/BehaviorTreeEditor/Private/BehaviorTreeGraphNode.cpp
--- a/Engine/Source/Editor/BehaviorTreeEditor/Private/BehaviorTreeGraphNode.cpp
+++ b/Engine/Source/Editor/BehaviorTreeEditor/Private/BehaviorTreeGraphNode.cpp
@@ -333,6 +333,19 @@ void UBehaviorTreeGraphNode::InsertSubNodeAt(UAIGraphNode* SubNode, int32 DropIn
 	}
 }
 
+void UBehaviorTreeGraphNode::UpdateErrorMessage()
+{
+	Super::UpdateErrorMessage();
+
+	if (ErrorMessage.IsEmpty())
+	{
+		if (const UBTNode* BTNodeInstance = Cast<UBTNode>(NodeInstance))
+		{
+			ErrorMessage = BTNodeInstance->GetErrorMessage();
+		}
+	}
+}
+
 FLinearColor UBehaviorTreeGraphNode::GetBackgroundColor(bool bIsActiveForDebugger) const
 {
 	return BehaviorTreeColors::NodeBody::Default;
diff --git a/Engine/Source/Editor/BehaviorTreeEditor/Private/SGraphNode_BehaviorTree.cpp b/Engine/Source/Editor/BehaviorTreeEditor/Private/SGraphNode_BehaviorTree.cpp
--- a/Engine/Source/Editor/BehaviorTreeEditor/Private/SGraphNode_BehaviorTree.cpp
+++ b/Engine/Source/Editor/BehaviorTreeEditor/Private/SGraphNode_BehaviorTree.cpp
@@ -783,35 +783,13 @@ void SGraphNode_BehaviorTree::AddPin(const TSharedRef<SGraphPin>& PinToAdd)
 
 TSharedPtr<SToolTip> SGraphNode_BehaviorTree::GetComplexTooltip()
 {
-	UBehaviorTreeGraphNode_CompositeDecorator* DecoratorNode = Cast<UBehaviorTreeGraphNode_CompositeDecorator>(GraphNode);
-	if (DecoratorNode && DecoratorNode->GetBoundGraph())
-	{
-		return SNew(SToolTip)
-			[
-				SNew(SOverlay)
-				+SOverlay::Slot()
-				[
-					// Create the tooltip graph preview, make sure to disable state overlays to
-					// prevent the PIE / read-only borders from obscuring the graph
-					SNew(SGraphPreviewer, DecoratorNode->GetBoundGraph())
-					.CornerOverlayText(LOCTEXT("CompositeDecoratorOverlayText", "Composite Decorator"))
-					.ShowGraphStateOverlay(false)
-				]
-				+SOverlay::Slot()
-				.Padding(2.0f)
-				[
-					SNew(STextBlock)
-					.Text(LOCTEXT("CompositeDecoratorTooltip", "Double-click to Open"))
-					.ColorAndOpacity(FSlateColor::UseSubduedForeground())
-				]
-			];
-	}
+	const UBehaviorTreeGraphNode* BTGraphNode = Cast<UBehaviorTreeGraphNode>(GraphNode);
+	const bool bHasErrors = BTGraphNode && BTGraphNode->HasErrors();
 
-	UBehaviorTreeGraphNode_Task* TaskNode = Cast<UBehaviorTreeGraphNode_Task>(GraphNode);
-	if(TaskNode && TaskNode->NodeInstance)
+	if (!bHasErrors)
 	{
-		UBTTask_RunBehavior* RunBehavior = Cast<UBTTask_RunBehavior>(TaskNode->NodeInstance);
-		if(RunBehavior && RunBehavior->GetSubtreeAsset() && RunBehavior->GetSubtreeAsset()->BTGraph)
+		UBehaviorTreeGraphNode_CompositeDecorator* DecoratorNode = Cast<UBehaviorTreeGraphNode_CompositeDecorator>(GraphNode);
+		if (DecoratorNode && DecoratorNode->GetBoundGraph())
 		{
 			return SNew(SToolTip)
 				[
@@ -820,19 +798,47 @@ TSharedPtr<SToolTip> SGraphNode_BehaviorTree::GetComplexTooltip()
 					[
 						// Create the tooltip graph preview, make sure to disable state overlays to
 						// prevent the PIE / read-only borders from obscuring the graph
-						SNew(SGraphPreviewer, RunBehavior->GetSubtreeAsset()->BTGraph)
-						.CornerOverlayText(LOCTEXT("RunBehaviorOverlayText", "Run Behavior"))
+						SNew(SGraphPreviewer, DecoratorNode->GetBoundGraph())
+						.CornerOverlayText(LOCTEXT("CompositeDecoratorOverlayText", "Composite Decorator"))
 						.ShowGraphStateOverlay(false)
 					]
 					+SOverlay::Slot()
 					.Padding(2.0f)
 					[
 						SNew(STextBlock)
-						.Text(LOCTEXT("RunBehaviorTooltip", "Double-click to Open"))
+						.Text(LOCTEXT("CompositeDecoratorTooltip", "Double-click to Open"))
 						.ColorAndOpacity(FSlateColor::UseSubduedForeground())
 					]
 				];
 		}
+
+		UBehaviorTreeGraphNode_Task* TaskNode = Cast<UBehaviorTreeGraphNode_Task>(GraphNode);
+		if(TaskNode && TaskNode->NodeInstance)
+		{
+			UBTTask_RunBehavior* RunBehavior = Cast<UBTTask_RunBehavior>(TaskNode->NodeInstance);
+			if(RunBehavior && RunBehavior->GetSubtreeAsset() && RunBehavior->GetSubtreeAsset()->BTGraph)
+			{
+				return SNew(SToolTip)
+					[
+						SNew(SOverlay)
+						+SOverlay::Slot()
+						[
+							// Create the tooltip graph preview, make sure to disable state overlays to
+							// prevent the PIE / read-only borders from obscuring the graph
+							SNew(SGraphPreviewer, RunBehavior->GetSubtreeAsset()->BTGraph)
+							.CornerOverlayText(LOCTEXT("RunBehaviorOverlayText", "Run Behavior"))
+							.ShowGraphStateOverlay(false)
+						]
+						+SOverlay::Slot()
+						.Padding(2.0f)
+						[
+							SNew(STextBlock)
+							.Text(LOCTEXT("RunBehaviorTooltip", "Double-click to Open"))
+							.ColorAndOpacity(FSlateColor::UseSubduedForeground())
+						]
+					];
+			}
+		}
 	}
 
 	return IDocumentation::Get()->CreateToolTip(TAttribute<FText>(this, &SGraphNode::GetNodeTooltip), NULL, GraphNode->GetDocumentationLink(), GraphNode->GetDocumentationExcerptName());
diff --git a/Engine/Source/Runtime/AIModule/Classes/BehaviorTree/BTNode.h b/Engine/Source/Runtime/AIModule/Classes/BehaviorTree/BTNode.h
--- a/Engine/Source/Runtime/AIModule/Classes/BehaviorTree/BTNode.h
+++ b/Engine/Source/Runtime/AIModule/Classes/BehaviorTree/BTNode.h
@@ -157,6 +157,9 @@ class UBTNode : public UObject, public IGameplayTaskOwnerInterface
 	/** Get whether this node is using a blueprint for its logic */
 	AIMODULE_API virtual bool UsesBlueprint() const;
 
+	/** Validates this node and returns the resulting error message. Empty means no error. */
+	AIMODULE_API virtual FString GetErrorMessage() const;
+
 	/** Called after creating new node in behavior tree editor, use for versioning */
 	virtual void OnNodeCreated() {}
 #endif
diff --git a/Engine/Source/Runtime/AIModule/Private/BehaviorTree/BTNode.cpp b/Engine/Source/Runtime/AIModule/Private/BehaviorTree/BTNode.cpp
--- a/Engine/Source/Runtime/AIModule/Private/BehaviorTree/BTNode.cpp
+++ b/Engine/Source/Runtime/AIModule/Private/BehaviorTree/BTNode.cpp
@@ -221,6 +221,11 @@ bool UBTNode::UsesBlueprint() const
 	return false;
 }
 
+FString UBTNode::GetErrorMessage() const
+{
+	return FString();
+}
+
 #endif
 
 UGameplayTasksComponent* UBTNode::GetGameplayTasksComponent(const UGameplayTask& Task) const
