[Behavior Tree] Fix BT Composite Decorator had an erroneous tooltip ("Class not found, make sure it's saved!")

diff --git a/Engine/Source/Editor/BehaviorTreeEditor/Classes/BehaviorTreeGraphNode_CompositeDecorator.h b/Engine/Source/Editor/BehaviorTreeEditor/Classes/BehaviorTreeGraphNode_CompositeDecorator.h
--- a/Engine/Source/Editor/BehaviorTreeEditor/Classes/BehaviorTreeGraphNode_CompositeDecorator.h
+++ b/Engine/Source/Editor/BehaviorTreeEditor/Classes/BehaviorTreeGraphNode_CompositeDecorator.h
@@ -48,6 +48,7 @@ class BEHAVIORTREEEDITOR_API UBehaviorTreeGraphNode_CompositeDecorator : public
 	virtual FText GetNodeTitle(ENodeTitleType::Type TitleType) const override;
 	virtual void AllocateDefaultPins() override;
 	virtual FText GetDescription() const override;
+	virtual FText GetTooltipText() const override;
 	virtual void PostPlacedNewNode() override;
 	virtual void PostLoad() override;
 	virtual UEdGraph* GetBoundGraph() const override { return BoundGraph; }
diff --git a/Engine/Source/Editor/BehaviorTreeEditor/Private/BehaviorTreeGraphNode_CompositeDecorator.cpp b/Engine/Source/Editor/BehaviorTreeEditor/Private/BehaviorTreeGraphNode_CompositeDecorator.cpp
--- a/Engine/Source/Editor/BehaviorTreeEditor/Private/BehaviorTreeGraphNode_CompositeDecorator.cpp
+++ b/Engine/Source/Editor/BehaviorTreeEditor/Private/BehaviorTreeGraphNode_CompositeDecorator.cpp
@@ -64,6 +64,16 @@ FText UBehaviorTreeGraphNode_CompositeDecorator::GetDescription() const
 	return FText::FromString(CachedDescription);
 }
 
+FText UBehaviorTreeGraphNode_CompositeDecorator::GetTooltipText() const
+{
+	if (ErrorMessage.IsEmpty() == false)
+	{
+		return FText::FromString(ErrorMessage);
+	}
+
+	return LOCTEXT("CompositeTooltip", "This node enables you to set up more advanced conditions using logic gates.");
+}
+
 void UBehaviorTreeGraphNode_CompositeDecorator::PostPlacedNewNode()
 {
 	if (BoundGraph == nullptr)
