Improve behavior tree debugging
- Add a warning to the bottom left of the debug window when the behavior tree being debugged is paused.
- Gather runtime values from instanced nodes when needed so the correct values are displayed when debugging.
- Make double clicking on a dynamic subtree node open the running subtree

diff --git a/Engine/Source/Editor/BehaviorTreeEditor/Private/BehaviorTreeDebugger.cpp b/Engine/Source/Editor/BehaviorTreeEditor/Private/BehaviorTreeDebugger.cpp
--- a/Engine/Source/Editor/BehaviorTreeEditor/Private/BehaviorTreeDebugger.cpp
+++ b/Engine/Source/Editor/BehaviorTreeEditor/Private/BehaviorTreeDebugger.cpp
@@ -11,6 +11,7 @@
 #include "BehaviorTree/BTNode.h"
 #include "BehaviorTree/BTTaskNode.h"
 #include "BehaviorTree/BTAuxiliaryNode.h"
+#include "BehaviorTree/Tasks/BTTask_RunBehaviorDynamic.h"
 #include "BehaviorTreeGraphNode_CompositeDecorator.h"
 #include "BehaviorTreeEditor.h"
 #include "Editor/UnrealEdEngine.h"
@@ -855,6 +856,17 @@ int32 FBehaviorTreeDebugger::GetShownStateIndex() const
 	return 0;
 }
 
+bool FBehaviorTreeDebugger::IsBehaviorExecutionPaused() const
+{
+#if USE_BEHAVIORTREE_DEBUGGER
+	if (TreeInstance.IsValid() && TreeInstance->DebuggerSteps.IsValidIndex(ActiveStepIndex))
+	{
+		return TreeInstance->DebuggerSteps[ActiveStepIndex].bIsExecutionPaused;
+	}
+#endif
+	return false;
+}
+
 void FBehaviorTreeDebugger::StepForwardInto()
 {
 #if USE_BEHAVIORTREE_DEBUGGER
@@ -1197,6 +1209,26 @@ void FBehaviorTreeDebugger::UpdateDebuggerViewOnTick()
 #endif
 }
 
+class UBehaviorTree* FBehaviorTreeDebugger::GetDynamicSubtreeTaskBehaviorTree(const UBTTask_RunBehaviorDynamic* Node) const
+{
+	if (UBehaviorTreeComponent* TreeComp = TreeInstance.Get())
+	{
+		if (TreeComp->DebuggerSteps.IsValidIndex(ActiveStepIndex))
+		{
+			const FBehaviorTreeExecutionStep& ActiveStep = TreeComp->DebuggerSteps[ActiveStepIndex];
+			if (ActiveStep.InstanceStack.IsValidIndex(DebuggerInstanceIndex))
+			{
+				const FBehaviorTreeDebuggerInstance& DebugInstance = ActiveStep.InstanceStack[DebuggerInstanceIndex];
+				if (DebugInstance.RuntimeDesc.IsValidIndex(Node->GetExecutionIndex()))
+				{
+					return Node->GetBehaviorAssetFromRuntimeValue(DebugInstance.RuntimeDesc[Node->GetExecutionIndex()]);
+				}
+			}
+		}
+	}
+	return nullptr;
+}
+
 FText FBehaviorTreeDebugger::FindValueForKey(const FName& InKeyName, bool bUseCurrentState) const
 {
 #if USE_BEHAVIORTREE_DEBUGGER
diff --git a/Engine/Source/Editor/BehaviorTreeEditor/Private/BehaviorTreeDebugger.h b/Engine/Source/Editor/BehaviorTreeEditor/Private/BehaviorTreeDebugger.h
--- a/Engine/Source/Editor/BehaviorTreeEditor/Private/BehaviorTreeDebugger.h
+++ b/Engine/Source/Editor/BehaviorTreeEditor/Private/BehaviorTreeDebugger.h
@@ -47,6 +47,7 @@ public:
 	bool IsDebuggerRunning() const;
 	bool IsShowingCurrentState() const;
 	int32 GetShownStateIndex() const;
+	bool IsBehaviorExecutionPaused() const;
 
 	void OnObjectSelected(UObject* Object);
 	void OnAIDebugSelected(const APawn* Pawn);
@@ -86,6 +87,8 @@ public:
 	bool HasContinuousNextStep() const;
 	bool HasContinuousPrevStep() const;
 
+	UBehaviorTree* GetDynamicSubtreeTaskBehaviorTree(const UBTTask_RunBehaviorDynamic* Node) const;
+
 	/**
 	 * Find a (display) value for a given key.
 	 * @param	InKeyName			Key to find a value for
diff --git a/Engine/Source/Editor/BehaviorTreeEditor/Private/BehaviorTreeEditor.cpp b/Engine/Source/Editor/BehaviorTreeEditor/Private/BehaviorTreeEditor.cpp
--- a/Engine/Source/Editor/BehaviorTreeEditor/Private/BehaviorTreeEditor.cpp
+++ b/Engine/Source/Editor/BehaviorTreeEditor/Private/BehaviorTreeEditor.cpp
@@ -16,6 +16,7 @@
 #include "BehaviorTree/Services/BTService_BlueprintBase.h"
 #include "BehaviorTree/Tasks/BTTask_BlueprintBase.h"
 #include "BehaviorTree/Tasks/BTTask_RunBehavior.h"
+#include "BehaviorTree/Tasks/BTTask_RunBehaviorDynamic.h"
 #include "BehaviorTreeColors.h"
 #include "BehaviorTreeDebugger.h"
 #include "BehaviorTreeDecoratorGraphNode_Logic.h"
@@ -532,6 +533,11 @@ FGraphAppearanceInfo FBehaviorTreeEditor::GetGraphAppearance() const
 	{
 		AppearanceInfo.PIENotifyText = LOCTEXT("PausedLabel", "PAUSED");
 	}
+
+	if (Debugger.IsValid() && Debugger->IsBehaviorExecutionPaused())
+	{
+		AppearanceInfo.WarningText = LOCTEXT("BehaviorExecutionPausedLabel", "BEHAVIOR EXECUTION PAUSED");
+	}
 	
 	return AppearanceInfo;
 }
@@ -1249,19 +1255,34 @@ void FBehaviorTreeEditor::OnNodeDoubleClicked(class UEdGraphNode* Node)
 			}
 		}
 	}
-	else if (UBehaviorTreeGraphNode_SubtreeTask* Task = Cast<UBehaviorTreeGraphNode_SubtreeTask>(Node))
+	else if (UBehaviorTreeGraphNode_Task* Task = Cast<UBehaviorTreeGraphNode_Task>(Node))
 	{
-		if (UBTTask_RunBehavior* RunTask = Cast<UBTTask_RunBehavior>(Task->NodeInstance))
+		UBehaviorTree* SubTreeToOpen = nullptr;
+		if (UBTTask_RunBehavior* SubtreeTask = Cast<UBTTask_RunBehavior>(Task->NodeInstance))
+		{
+			SubTreeToOpen = SubtreeTask->GetSubtreeAsset();
+		}
+		else if (UBTTask_RunBehaviorDynamic* DynamicSubTreeTask = Cast<UBTTask_RunBehaviorDynamic>(Task->NodeInstance))
 		{
-			if (RunTask->GetSubtreeAsset())
+			if (Debugger)
 			{
-				GEditor->GetEditorSubsystem<UAssetEditorSubsystem>()->OpenEditorForAsset(RunTask->GetSubtreeAsset());
+				SubTreeToOpen = Debugger->GetDynamicSubtreeTaskBehaviorTree(DynamicSubTreeTask);
+			}
+			
+			if (!SubTreeToOpen)
+			{
+				SubTreeToOpen = DynamicSubTreeTask->GetDefaultBehaviorAsset();
+			}
+		}
 
-				IBehaviorTreeEditor* ChildNodeEditor = static_cast<IBehaviorTreeEditor*>(GEditor->GetEditorSubsystem<UAssetEditorSubsystem>()->FindEditorForAsset(RunTask->GetSubtreeAsset(), true));
-				if (ChildNodeEditor)
-				{
-					ChildNodeEditor->InitializeDebuggerState(Debugger.Get());
-				}
+		if (SubTreeToOpen)
+		{
+			GEditor->GetEditorSubsystem<UAssetEditorSubsystem>()->OpenEditorForAsset(SubTreeToOpen);
+
+			IBehaviorTreeEditor* ChildNodeEditor = static_cast<IBehaviorTreeEditor*>(GEditor->GetEditorSubsystem<UAssetEditorSubsystem>()->FindEditorForAsset(SubTreeToOpen, true));
+			if (ChildNodeEditor)
+			{
+				ChildNodeEditor->InitializeDebuggerState(Debugger.Get());
 			}
 		}
 	}
diff --git a/Engine/Source/Runtime/AIModule/Classes/BehaviorTree/BehaviorTreeComponent.h b/Engine/Source/Runtime/AIModule/Classes/BehaviorTree/BehaviorTreeComponent.h
--- a/Engine/Source/Runtime/AIModule/Classes/BehaviorTree/BehaviorTreeComponent.h
+++ b/Engine/Source/Runtime/AIModule/Classes/BehaviorTree/BehaviorTreeComponent.h
@@ -230,7 +230,7 @@ public:
 	AIMODULE_API uint16 GetActiveInstanceIdx() const;
 
 	/** @return node memory */
-	AIMODULE_API uint8* GetNodeMemory(UBTNode* Node, int32 InstanceIdx) const;
+	AIMODULE_API uint8* GetNodeMemory(const UBTNode* Node, int32 InstanceIdx) const;
 
 	/** @return true if ExecutionRequest is switching to higher priority node */
 	AIMODULE_API bool IsRestartPending() const;
diff --git a/Engine/Source/Runtime/AIModule/Classes/BehaviorTree/BehaviorTreeTypes.h b/Engine/Source/Runtime/AIModule/Classes/BehaviorTree/BehaviorTreeTypes.h
--- a/Engine/Source/Runtime/AIModule/Classes/BehaviorTree/BehaviorTreeTypes.h
+++ b/Engine/Source/Runtime/AIModule/Classes/BehaviorTree/BehaviorTreeTypes.h
@@ -274,6 +274,9 @@ struct FBehaviorTreeExecutionStep
 
 	/** Id of execution step */
 	int32 ExecutionStepId;
+
+	/** If true, the behavior was paused in this execution step. */
+	bool bIsExecutionPaused = false;
 };
 
 /** identifier of subtree instance */
diff --git a/Engine/Source/Runtime/AIModule/Classes/BehaviorTree/Tasks/BTTask_RunBehaviorDynamic.h b/Engine/Source/Runtime/AIModule/Classes/BehaviorTree/Tasks/BTTask_RunBehaviorDynamic.h
--- a/Engine/Source/Runtime/AIModule/Classes/BehaviorTree/Tasks/BTTask_RunBehaviorDynamic.h
+++ b/Engine/Source/Runtime/AIModule/Classes/BehaviorTree/Tasks/BTTask_RunBehaviorDynamic.h
@@ -29,6 +29,7 @@ class UBTTask_RunBehaviorDynamic : public UBTTaskNode
 
 #if WITH_EDITOR
 	AIMODULE_API virtual FName GetNodeIconName() const override;
+	AIMODULE_API UBehaviorTree* GetBehaviorAssetFromRuntimeValue(const FString& RuntimeValue) const;
 #endif // WITH_EDITOR
 
 	AIMODULE_API bool HasMatchingTag(const FGameplayTag& Tag) const;
diff --git a/Engine/Source/Runtime/AIModule/Private/BehaviorTree/BehaviorTreeComponent.cpp b/Engine/Source/Runtime/AIModule/Private/BehaviorTree/BehaviorTreeComponent.cpp
--- a/Engine/Source/Runtime/AIModule/Private/BehaviorTree/BehaviorTreeComponent.cpp
+++ b/Engine/Source/Runtime/AIModule/Private/BehaviorTree/BehaviorTreeComponent.cpp
@@ -169,6 +169,9 @@ void UBehaviorTreeComponent::PauseLogic(const FString& Reason)
 	{
 		BlackboardComp->PauseObserverNotifications();
 	}
+
+	// Store a debugger execution step so debugger is aware we've been paused.
+	StoreDebuggerExecutionStep(EBTExecutionSnap::Regular);
 }
 
 EAILogicResuming::Type UBehaviorTreeComponent::ResumeLogic(const FString& Reason)
@@ -2874,7 +2877,7 @@ UBTNode* UBehaviorTreeComponent::FindTemplateNode(const UBTNode* Node) const
 	return NULL;
 }
 
-uint8* UBehaviorTreeComponent::GetNodeMemory(UBTNode* Node, int32 InstanceIdx) const
+uint8* UBehaviorTreeComponent::GetNodeMemory(const UBTNode* Node, int32 InstanceIdx) const
 {
 	return InstanceStack.IsValidIndex(InstanceIdx) ? (uint8*)Node->GetNodeMemory<uint8>(InstanceStack[InstanceIdx]) : NULL;
 }
@@ -3194,6 +3197,7 @@ void UBehaviorTreeComponent::StoreDebuggerExecutionStep(EBTExecutionSnap::Type S
 	}
 
 	FBehaviorTreeExecutionStep CurrentStep;
+	CurrentStep.bIsExecutionPaused = bIsPaused;
 	CurrentStep.ExecutionStepId = DebuggerSteps.Num() ? DebuggerSteps.Last().ExecutionStepId + 1 : 0;
 	CurrentStep.TimeStamp = GetWorld()->GetTimeSeconds();
 	CurrentStep.BlackboardValues = SearchStartBlackboard;
@@ -3391,7 +3395,19 @@ void UBehaviorTreeComponent::StoreDebuggerRuntimeValues(TArray<FString>& Runtime
 		uint8* NodeMemory = (uint8*)Node->GetNodeMemory<uint8>(InstanceInfo);
 
 		RuntimeValues.Reset();
-		Node->DescribeRuntimeValues(*this, NodeMemory, EBTDescriptionVerbosity::Basic, RuntimeValues);
+
+		constexpr EBTDescriptionVerbosity::Type DescriptionVerbosity = EBTDescriptionVerbosity::Basic;
+		if (Node->HasInstance())
+		{
+			if (UBTNode* NodeInstance = Node->GetNodeInstance(*this, NodeMemory))
+			{
+				NodeInstance->DescribeRuntimeValues(*this, NodeMemory, DescriptionVerbosity, RuntimeValues);
+			}
+		}
+		else
+		{
+			Node->DescribeRuntimeValues(*this, NodeMemory, DescriptionVerbosity, RuntimeValues);
+		}
 
 		FString ComposedDesc;
 		for (int32 ValueIndex = 0; ValueIndex < RuntimeValues.Num(); ValueIndex++)
diff --git a/Engine/Source/Runtime/AIModule/Private/BehaviorTree/Tasks/BTTask_RunBehaviorDynamic.cpp b/Engine/Source/Runtime/AIModule/Private/BehaviorTree/Tasks/BTTask_RunBehaviorDynamic.cpp
--- a/Engine/Source/Runtime/AIModule/Private/BehaviorTree/Tasks/BTTask_RunBehaviorDynamic.cpp
+++ b/Engine/Source/Runtime/AIModule/Private/BehaviorTree/Tasks/BTTask_RunBehaviorDynamic.cpp
@@ -3,6 +3,7 @@
 #include "BehaviorTree/Tasks/BTTask_RunBehaviorDynamic.h"
 #include "VisualLogger/VisualLogger.h"
 #include "BehaviorTree/BehaviorTree.h"
+#include "Internationalization/Regex.h"
 
 #include UE_INLINE_GENERATED_CPP_BY_NAME(BTTask_RunBehaviorDynamic)
 
@@ -59,6 +60,10 @@ void UBTTask_RunBehaviorDynamic::DescribeRuntimeValues(const UBehaviorTreeCompon
 {
 	Super::DescribeRuntimeValues(OwnerComp, NodeMemory, Verbosity, Values);
 	Values.Add(FString::Printf(TEXT("subtree: %s"), *GetNameSafe(BehaviorAsset)));
+	if (BehaviorAsset)
+	{
+		Values.Add(FString::Printf(TEXT("subtree path: %s"), *BehaviorAsset->GetPathName()));
+	}
 }
 
 bool UBTTask_RunBehaviorDynamic::SetBehaviorAsset(UBehaviorTree* NewBehaviorAsset)
@@ -72,6 +77,18 @@ bool UBTTask_RunBehaviorDynamic::SetBehaviorAsset(UBehaviorTree* NewBehaviorAsse
 }
 
 #if WITH_EDITOR
+UBehaviorTree* UBTTask_RunBehaviorDynamic::GetBehaviorAssetFromRuntimeValue(const FString& RuntimeValue) const
+{
+	static const FRegexPattern RegexPatern(FString(TEXT("[\\n\\r].*subtree path:\\s*([^\\n\\r]*)")));
+	FRegexMatcher Matcher(RegexPatern, RuntimeValue);
+	if (Matcher.FindNext())
+	{
+		// Capture group 0 is the whole subtree line, capture group 1 is the path.
+		const FString SubtreePath = Matcher.GetCaptureGroup(1);
+		return FindObject<UBehaviorTree>(nullptr, *SubtreePath);
+	}
+	return nullptr;
+}
 
 FName UBTTask_RunBehaviorDynamic::GetNodeIconName() const
 {
