Behavior Tree Editor: Allow search to find decorators and services by the class of the instance

diff --git a/Engine/Source/Editor/BehaviorTreeEditor/Private/FindInBT.cpp b/Engine/Source/Editor/BehaviorTreeEditor/Private/FindInBT.cpp
--- a/Engine/Source/Editor/BehaviorTreeEditor/Private/FindInBT.cpp
+++ b/Engine/Source/Editor/BehaviorTreeEditor/Private/FindInBT.cpp
@@ -300,7 +300,7 @@ void SFindInBT::MatchTokensInChild(const TArray<FString>& Tokens, UBehaviorTreeG
 	}
 
 	FString ChildName = Child->GetNodeTitle(ENodeTitleType::ListView).ToString();
-	FString ChildSearchString = ChildName + Child->GetClass()->GetName() + Child->NodeComment;
+	FString ChildSearchString = ChildName + Child->GetClass()->GetName() + Child->NodeComment + GetNameSafe(Child->NodeInstance ? Child->NodeInstance->GetClass() : nullptr);
 	ChildSearchString = ChildSearchString.Replace(TEXT(" "), TEXT(""));
 	if (StringMatchesSearchTokens(Tokens, ChildSearchString))
 	{
