Behavior Tree Debugger: Fix subtrees never switched back to inactive state.

diff --git a/Engine/Source/Editor/BehaviorTreeEditor/Private/BehaviorTreeDebugger.cpp b/Engine/Source/Editor/BehaviorTreeEditor/Private/BehaviorTreeDebugger.cpp
--- a/Engine/Source/Editor/BehaviorTreeEditor/Private/BehaviorTreeDebugger.cpp
+++ b/Engine/Source/Editor/BehaviorTreeEditor/Private/BehaviorTreeDebugger.cpp
@@ -182,6 +182,26 @@ void FBehaviorTreeDebugger::Tick(float DeltaTime)
 
 		if (Step.ExecutionStepId > DisplayedExecutionStepId)
 		{
+			// If we're in an inactive step, and this tree does not appear in the current step, then we skip updating
+			// the active step index. Without this, UpdateDebuggerInstance would call ClearDebuggerState (also resetting
+			// ActiveStepIndex) as DebuggerInstanceIndex became invalid, but in the next frame Clear would not be
+			// called, as DebuggerInstanceIndex hasn't changed (so ActiveStepIndex would be unchanged).
+			// This inconsistency caused a flash of the state label, and then switched to SIMULATION in the next frame,
+			// even though the subtree was not reachable anymore.
+			if (ActiveStepIndex == INDEX_NONE)
+			{
+				const bool bAppearsInStep = Step.InstanceStack.ContainsByPredicate([this](
+					const FBehaviorTreeDebuggerInstance& Instance)
+				{
+					return Instance.TreeAsset == TreeAsset;
+				});
+
+				if (!bAppearsInStep)
+				{
+					continue;
+				}
+			}
+
 			ActiveStepIndex = i;
 			LastValidExecutionStepId = Step.ExecutionStepId;
 
@@ -355,7 +375,7 @@ void FBehaviorTreeDebugger::ClearDebuggerState(bool bKeepSubtree)
 	DisplayedExecutionStepId = FBehaviorTreeExecutionStep::InvalidExecutionId;
 
 	DebuggerInstanceIndex = INDEX_NONE;
-	ActiveStepIndex = 0;
+	ActiveStepIndex = INDEX_NONE;
 
 	if (TreeAsset && RootNode.IsValid())
 	{
