Behavior Tree: Show runtime node description on the node itself (instead of its tooltip)

diff --git a/Engine/Source/Editor/BehaviorTreeEditor/Private/BehaviorTreeGraphNode.cpp b/Engine/Source/Editor/BehaviorTreeEditor/Private/BehaviorTreeGraphNode.cpp
--- a/Engine/Source/Editor/BehaviorTreeEditor/Private/BehaviorTreeGraphNode.cpp
+++ b/Engine/Source/Editor/BehaviorTreeEditor/Private/BehaviorTreeGraphNode.cpp
@@ -82,10 +82,6 @@ FText UBehaviorTreeGraphNode::GetTooltipText() const
 		{
 			TooltipDesc = LOCTEXT("ObserverError", "Observer has invalid abort setting!");
 		}
-		else if (DebuggerRuntimeDescription.Len() > 0)
-		{
-			TooltipDesc = FText::FromString(DebuggerRuntimeDescription);
-		}
 	}
 
 	if (TooltipDesc.IsEmpty())
diff --git a/Engine/Source/Editor/BehaviorTreeEditor/Private/SGraphNode_BehaviorTree.cpp b/Engine/Source/Editor/BehaviorTreeEditor/Private/SGraphNode_BehaviorTree.cpp
--- a/Engine/Source/Editor/BehaviorTreeEditor/Private/SGraphNode_BehaviorTree.cpp
+++ b/Engine/Source/Editor/BehaviorTreeEditor/Private/SGraphNode_BehaviorTree.cpp
@@ -477,6 +477,14 @@ void SGraphNode_BehaviorTree::UpdateGraphNode()
 										.Visibility(this, &SGraphNode_BehaviorTree::GetDescriptionVisibility)
 										.Text(this, &SGraphNode_BehaviorTree::GetDescription)
 									]
+									+SVerticalBox::Slot()
+									.AutoHeight()
+									[
+										SNew(STextBlock)
+										.Visibility(this, &SGraphNode_BehaviorTree::GetRuntimeDescriptionVisibility)
+										.Text(this, &SGraphNode_BehaviorTree::GetRuntimeDescription)
+										.TextStyle(FAppStyle::Get(), "SmallText")
+									]
 								]
 								+SOverlay::Slot()
 								.HAlign(HAlign_Right)
@@ -713,6 +721,21 @@ FText SGraphNode_BehaviorTree::GetPinTooltip(UEdGraphPin* GraphPinObj) const
 	return HoverText;
 }
 
+FText SGraphNode_BehaviorTree::GetRuntimeDescription() const
+{
+	const UBehaviorTreeGraphNode* MyNode = CastChecked<UBehaviorTreeGraphNode>(GraphNode);
+	return MyNode ? FText::FromString(MyNode->DebuggerRuntimeDescription) : FText::GetEmpty();
+}
+
+EVisibility SGraphNode_BehaviorTree::GetRuntimeDescriptionVisibility() const
+{
+	const UBehaviorTreeGraphNode* MyNode = CastChecked<UBehaviorTreeGraphNode>(GraphNode);
+
+	return GetDescriptionVisibility().IsVisible() && !MyNode->DebuggerRuntimeDescription.IsEmpty()
+		? EVisibility::HitTestInvisible
+		: EVisibility::Collapsed;
+}
+
 void SGraphNode_BehaviorTree::CreatePinWidgets()
 {
 	UBehaviorTreeGraphNode* StateNode = CastChecked<UBehaviorTreeGraphNode>(GraphNode);
diff --git a/Engine/Source/Editor/BehaviorTreeEditor/Private/SGraphNode_BehaviorTree.h b/Engine/Source/Editor/BehaviorTreeEditor/Private/SGraphNode_BehaviorTree.h
--- a/Engine/Source/Editor/BehaviorTreeEditor/Private/SGraphNode_BehaviorTree.h
+++ b/Engine/Source/Editor/BehaviorTreeEditor/Private/SGraphNode_BehaviorTree.h
@@ -116,6 +116,9 @@ protected:
 	virtual const FSlateBrush* GetNameIcon() const override;
 	virtual EVisibility GetBlueprintIconVisibility() const;
 
+	FText GetRuntimeDescription() const;
+	EVisibility GetRuntimeDescriptionVisibility() const;
+
 	/** Get the visibility of the index overlay */
 	EVisibility GetIndexVisibility() const;
 
