PR #11456: Add ability to search for blackboard keys and gameplay tags in the behavior tree editor

diff --git a/Engine/Source/Editor/BehaviorTreeEditor/BehaviorTreeEditor.Build.cs b/Engine/Source/Editor/BehaviorTreeEditor/BehaviorTreeEditor.Build.cs
--- a/Engine/Source/Editor/BehaviorTreeEditor/BehaviorTreeEditor.Build.cs
+++ b/Engine/Source/Editor/BehaviorTreeEditor/BehaviorTreeEditor.Build.cs
@@ -17,6 +17,7 @@ public class BehaviorTreeEditor : ModuleRules
 				"CoreUObject",
 				"Engine",
 				"GraphEditor",
+				"GameplayTags",
 				"InputCore",
 				"KismetWidgets",
 				"Slate",
diff --git a/Engine/Source/Editor/BehaviorTreeEditor/Private/FindInBT.cpp b/Engine/Source/Editor/BehaviorTreeEditor/Private/FindInBT.cpp
--- a/Engine/Source/Editor/BehaviorTreeEditor/Private/FindInBT.cpp
+++ b/Engine/Source/Editor/BehaviorTreeEditor/Private/FindInBT.cpp
@@ -1,6 +1,7 @@
 // Copyright Epic Games, Inc. All Rights Reserved.
 #include "FindInBT.h"
 
+#include "BehaviorTree/BehaviorTreeTypes.h"
 #include "BehaviorTreeEditor.h"
 #include "BehaviorTreeGraphNode.h"
 #include "BehaviorTreeGraphNode_Decorator.h"
@@ -9,6 +10,7 @@
 #include "EdGraph/EdGraphNode.h"
 #include "Framework/Application/SlateApplication.h"
 #include "Framework/Views/ITypedTableView.h"
+#include "GameplayTagContainer.h"
 #include "GraphEditor.h"
 #include "HAL/PlatformMath.h"
 #include "Input/Events.h"
@@ -26,6 +28,7 @@
 #include "UObject/ObjectPtr.h"
 #include "Widgets/Images/SImage.h"
 #include "Widgets/Input/SSearchBox.h"
+#include "Widgets/Input/STextComboBox.h"
 #include "Widgets/Layout/SBorder.h"
 #include "Widgets/Layout/SBox.h"
 #include "Widgets/SBoxPanel.h"
@@ -41,7 +44,7 @@ struct FSlateBrush;
 //////////////////////////////////////////////////////////////////////////
 // FFindInBTResult
 
-FFindInBTResult::FFindInBTResult(const FString& InValue) 
+FFindInBTResult::FFindInBTResult(const FString& InValue)
 	: Value(InValue), GraphNode(NULL)
 {
 }
@@ -83,7 +86,7 @@ TSharedRef<SWidget> FFindInBTResult::CreateIcon() const
 			Brush = FAppStyle::GetBrush(TEXT("GraphEditor.FIB_Event"));
 		}
 	}
-	
+
 	return SNew(SImage)
 		.Image(Brush)
 		.ColorAndOpacity(IconColor);
@@ -143,12 +146,18 @@ FString FFindInBTResult::GetCommentText() const
 
 //////////////////////////////////////////////////////////////////////////
 // SFindInBT
-
-void SFindInBT::Construct( const FArguments& InArgs, TSharedPtr<FBehaviorTreeEditor> InBehaviorTreeEditor)
+void SFindInBT::Construct(const FArguments& InArgs, TSharedPtr<FBehaviorTreeEditor> InBehaviorTreeEditor)
 {
 	BehaviorTreeEditorPtr = InBehaviorTreeEditor;
 
-	this->ChildSlot
+	// initialize slate with new combo box for different search types
+	SearchTypeComboBoxItems.Emplace(MakeShared<FString>(LOCTEXT("BTSearchType_Node", "Node").ToString()));
+	SearchTypeComboBoxItems.Emplace(MakeShared<FString>(LOCTEXT("BTSearchType_BlackboardKey", "Blackboard Key").ToString()));
+	SearchTypeComboBoxItems.Emplace(MakeShared<FString>(LOCTEXT("BTSearchType_GameplayTag", "GameplayTag").ToString()));
+
+	TSharedPtr<FString> CurrentlySelected = SearchTypeComboBoxItems[0];
+
+this->ChildSlot
 		[
 			SNew(SVerticalBox)
 			+SVerticalBox::Slot()
@@ -156,7 +165,16 @@ void SFindInBT::Construct( const FArguments& InArgs, TSharedPtr<FBehaviorTreeEdi
 			[
 				SNew(SHorizontalBox)
 				+SHorizontalBox::Slot()
-				.FillWidth(1)
+				.FillWidth(0.25f)
+				[
+					SNew(STextComboBox)
+					.OptionsSource(&SearchTypeComboBoxItems)
+					.InitiallySelectedItem(CurrentlySelected)
+					.OnSelectionChanged(this, &SFindInBT::OnSearchTypeSelectedItemChanged)
+					.Font(FAppStyle::GetFontStyle("PropertyWindow.NormalFont"))
+				]
+				+SHorizontalBox::Slot()
+				.FillWidth(0.75f)
 				[
 					SAssignNew(SearchTextField, SSearchBox)
 					.HintText(LOCTEXT("BehaviorTreeSearchHint", "Enter text to find nodes..."))
@@ -196,7 +214,7 @@ void SFindInBT::FocusForUse()
 void SFindInBT::OnSearchTextChanged(const FText& Text)
 {
 	SearchValue = Text.ToString();
-	
+
 	InitiateSearch();
 }
 
@@ -257,7 +275,7 @@ void SFindInBT::MatchTokens(const TArray<FString>& Tokens)
 	for (auto It(Graph->Nodes.CreateConstIterator()); It; ++It)
 	{
 		UEdGraphNode* Node = *It;
-			
+
 		const FString NodeName = Node->GetNodeTitle(ENodeTitleType::ListView).ToString();
 		FSearchResult NodeResult(new FFindInBTResult(NodeName, RootSearchResult, Node));
 
@@ -269,6 +287,14 @@ void SFindInBT::MatchTokens(const TArray<FString>& Tokens)
 		UBehaviorTreeGraphNode* BTNode = Cast<UBehaviorTreeGraphNode>(Node);
 		if (BTNode)
 		{
+			// search through node properties according to search type
+			FString OutExactFieldValueFound;
+			if (NodePropertyMatchesSearchTokens(Tokens, BTNode, OutExactFieldValueFound))
+			{
+				bNodeMatchesSearch = true;
+				NodeResult->ExactFieldValueFound = OutExactFieldValueFound;
+			}
+
 			// searching through nodes' decorators
 			for (auto DecoratorIt(BTNode->Decorators.CreateConstIterator()); DecoratorIt; ++DecoratorIt)
 			{
@@ -299,17 +325,20 @@ void SFindInBT::MatchTokensInChild(const TArray<FString>& Tokens, UBehaviorTreeG
 		return;
 	}
 
-	FString ChildName = Child->GetNodeTitle(ENodeTitleType::ListView).ToString();
+	FString OutExactFieldValueFound;
+
+	const FString ChildName = Child->GetNodeTitle(ENodeTitleType::ListView).ToString();
 	FString ChildSearchString = ChildName + Child->GetClass()->GetName() + Child->NodeComment + GetNameSafe(Child->NodeInstance ? Child->NodeInstance->GetClass() : nullptr);
 	ChildSearchString = ChildSearchString.Replace(TEXT(" "), TEXT(""));
-	if (StringMatchesSearchTokens(Tokens, ChildSearchString))
+	if (StringMatchesSearchTokens(Tokens, ChildSearchString) || NodePropertyMatchesSearchTokens(Tokens, Child, OutExactFieldValueFound))
 	{
 		FSearchResult DecoratorResult(new FFindInBTResult(ChildName, ParentNode, Child));
+		DecoratorResult->ExactFieldValueFound = OutExactFieldValueFound; // save field found to include it when generating search results
 		ParentNode->Children.Add(DecoratorResult);
 	}
 }
 
-TSharedRef<ITableRow> SFindInBT::OnGenerateRow( FSearchResult InItem, const TSharedRef<STableViewBase>& OwnerTable )
+TSharedRef<ITableRow> SFindInBT::OnGenerateRow(FSearchResult InItem, const TSharedRef<STableViewBase>& OwnerTable)
 {
 	return SNew(STableRow< TSharedPtr<FFindInBTResult> >, OwnerTable)
 		[
@@ -330,7 +359,7 @@ TSharedRef<ITableRow> SFindInBT::OnGenerateRow( FSearchResult InItem, const TSha
 					+SHorizontalBox::Slot()
 					.VAlign(VAlign_Center)
 					.AutoWidth()
-					.Padding(2, 0)
+					.Padding(2.f, 0.f)
 					[
 						SNew(STextBlock)
 						.Text(FText::FromString(InItem->Value))
@@ -341,14 +370,27 @@ TSharedRef<ITableRow> SFindInBT::OnGenerateRow( FSearchResult InItem, const TSha
 			+SHorizontalBox::Slot()
 			.AutoWidth()
 			.VAlign(VAlign_Center)
+			.Padding(10.f, 0.f)
 			[
 				SNew(STextBlock)
 				.Text(FText::FromString(InItem->GetNodeTypeText()))
 				.HighlightText(HighlightText)
 			]
+			// Include exact field found in search results
+			+SHorizontalBox::Slot()
+			.AutoWidth()
+			.VAlign(VAlign_Center)
+			.Padding(10.f, 0.f)
+			[
+				SNew(STextBlock)
+				.Text(FText::FromString(InItem->ExactFieldValueFound))
+				.ColorAndOpacity(FLinearColor::Green)
+				.HighlightText(HighlightText)
+			] 
 			+SHorizontalBox::Slot()
 			.HAlign(HAlign_Right)
 			.VAlign(VAlign_Center)
+			.Padding(5.f, 0.f)
 			[
 				SNew(STextBlock)
 				.Text(FText::FromString(InItem->GetCommentText()))
@@ -358,12 +400,12 @@ TSharedRef<ITableRow> SFindInBT::OnGenerateRow( FSearchResult InItem, const TSha
 		];
 }
 
-void SFindInBT::OnGetChildren(FSearchResult InItem, TArray< FSearchResult >& OutChildren)
+void SFindInBT::OnGetChildren(FSearchResult InItem, TArray<FSearchResult>& OutChildren)
 {
 	OutChildren += InItem->Children;
 }
 
-void SFindInBT::OnTreeSelectionChanged(FSearchResult Item , ESelectInfo::Type)
+void SFindInBT::OnTreeSelectionChanged(FSearchResult Item, ESelectInfo::Type)
 {
 	if (Item.IsValid())
 	{
@@ -375,7 +417,7 @@ bool SFindInBT::StringMatchesSearchTokens(const TArray<FString>& Tokens, const F
 {
 	bool bFoundAllTokens = true;
 
-	//search the entry for each token, it must have all of them to pass
+	// search the entry for each token, it must have all of them to pass
 	for (auto TokIT(Tokens.CreateConstIterator()); TokIT; ++TokIT)
 	{
 		const FString& Token = *TokIT;
@@ -388,6 +430,148 @@ bool SFindInBT::StringMatchesSearchTokens(const TArray<FString>& Tokens, const F
 	return bFoundAllTokens;
 }
 
+bool SFindInBT::NodePropertyMatchesSearchTokens(const TArray<FString>& Tokens, UBehaviorTreeGraphNode* Node, FString& OutExactFieldValueFound) const
+{
+	if (Node->NodeInstance == nullptr)
+	{
+		return false;
+	}
+
+	// nothing to search for, so abort early
+	if (SearchType == EFindInBTSearchType::Node)
+	{
+		return false;
+	}
+
+	for (FProperty* Property = Node->NodeInstance->GetClass()->PropertyLink;
+		 Property != nullptr; Property = Property->PropertyLinkNext)
+	{
+		if (FieldPropertyMatchesSearchTokens(Tokens, Property, Node->NodeInstance, OutExactFieldValueFound))
+		{
+			return true;
+		}
+	}
+	return false;
+}
+
+bool SFindInBT::FieldPropertyMatchesSearchTokens(const TArray<FString>& Tokens, const FProperty* Property, void* PropertySource, FString& OutExactFieldValueFound) const
+{
+	if (Property == nullptr)
+	{
+		return false;
+	}
+
+	if (PropertySource == nullptr)
+	{
+		return false;
+	}
+
+	const FArrayProperty* ArrayProp = CastField<FArrayProperty>(Property);
+	if (ArrayProp && ArrayProp->Inner)
+	{
+		FScriptArrayHelper ArrayHelper(ArrayProp, ArrayProp->ContainerPtrToValuePtr<void>(PropertySource));
+		for (int32 i = 0; i < ArrayHelper.Num(); ++i)
+		{
+			void* ArrayData = ArrayHelper.GetRawPtr(i);
+			if (FieldPropertyMatchesSearchTokens(Tokens, ArrayProp->Inner, ArrayData, OutExactFieldValueFound))
+			{
+				return true;
+			}
+		}
+	}
+
+	const FStructProperty* StructProp = CastField<FStructProperty>(Property);
+	if (StructProp && StructProp->Struct)
+	{
+		if (StructProp->Struct->IsChildOf(FBlackboardKeySelector::StaticStruct()))
+		{
+			if (SearchType != EFindInBTSearchType::BlackboardKey)
+			{
+				return false;
+			}
+
+			FBlackboardKeySelector* PropertyValue = Property->ContainerPtrToValuePtr<FBlackboardKeySelector>(PropertySource);
+			if (PropertyValue == nullptr)
+			{
+				return false;
+			}
+
+			if (StringMatchesSearchTokens(Tokens, PropertyValue->SelectedKeyName.ToString()))
+			{
+				OutExactFieldValueFound = PropertyValue->SelectedKeyName.ToString();
+				return true;
+			}
+		}
+		else if (StructProp->Struct->IsChildOf(FGameplayTag::StaticStruct()))
+		{
+			if (SearchType != EFindInBTSearchType::GameplayTag)
+			{
+				return false;
+			}
+
+			FGameplayTag* PropertyValue = Property->ContainerPtrToValuePtr<FGameplayTag>(PropertySource);
+			if (PropertyValue == nullptr)
+			{
+				return false;
+			}
+
+			if (StringMatchesSearchTokens(Tokens, PropertyValue->ToString()))
+			{
+				OutExactFieldValueFound = PropertyValue->ToString();
+				return true;
+			}
+		}
+		else
+		{
+			for (TFieldIterator<FProperty> iter(StructProp->Struct); iter; ++iter)
+			{
+				void* StructAddressInValuePtr = StructProp->ContainerPtrToValuePtr<void>(PropertySource);
+				if (StructAddressInValuePtr == nullptr)
+				{
+					continue;
+				}
+
+				if (FieldPropertyMatchesSearchTokens(Tokens, *iter, StructAddressInValuePtr, OutExactFieldValueFound))
+				{
+					return true;
+				}
+			}
+		}
+	}
+	else if (const FNameProperty* NameProperty = CastField<FNameProperty>(Property))
+	{
+		if (SearchType != EFindInBTSearchType::BlackboardKey)
+		{
+			return false;
+		}
+
+		FName* PropertyValue = NameProperty->ContainerPtrToValuePtr<FName>(PropertySource);
+		if (PropertyValue == nullptr)
+		{
+			return false;
+		}
+
+		if (StringMatchesSearchTokens(Tokens, PropertyValue->ToString()))
+		{
+			OutExactFieldValueFound = PropertyValue->ToString();
+			return true;
+		}
+	}
+
+	return false;
+}
+
+void SFindInBT::OnSearchTypeSelectedItemChanged(TSharedPtr<FString> NewValue, ESelectInfo::Type SelectInfo)
+{
+	const int32 IndexOf = SearchTypeComboBoxItems.IndexOfByKey(NewValue);
+	if (IndexOf != INDEX_NONE)
+	{
+		SearchType = static_cast<EFindInBTSearchType>(IndexOf);
+
+		InitiateSearch();
+	}
+}
+
 /////////////////////////////////////////////////////
 
 #undef LOCTEXT_NAMESPACE
diff --git a/Engine/Source/Editor/BehaviorTreeEditor/Private/FindInBT.h b/Engine/Source/Editor/BehaviorTreeEditor/Private/FindInBT.h
--- a/Engine/Source/Editor/BehaviorTreeEditor/Private/FindInBT.h
+++ b/Engine/Source/Editor/BehaviorTreeEditor/Private/FindInBT.h
@@ -28,6 +28,19 @@ class SWidget;
 class UBehaviorTreeGraphNode;
 class UEdGraphNode;
 
+/** Indicates the kind of search we want to do */
+enum class EFindInBTSearchType : uint8
+{
+	// Search for text on nodes
+	Node,
+
+	// Search through properties with matching Blackboard Key
+	BlackboardKey,
+
+	// Search through properties with matching GameplayTag
+	GameplayTag
+};
+
 /** Item that matched the search results */
 class FFindInBTResult
 {
@@ -64,6 +77,9 @@ public:
 
 	/** Search result parent */
 	TWeakPtr<FFindInBTResult> Parent;
+
+	/** Show in search result row the exact field value found */
+	FString ExactFieldValueFound;
 };
 
 /** Widget for searching for (BT nodes) across focused BehaviorTree */
@@ -109,6 +125,15 @@ private:
 	/** Determines if a string matches the search tokens */
 	static bool StringMatchesSearchTokens(const TArray<FString>& Tokens, const FString& ComparisonString);
 
+	/** Determines if a string matches a node property according to the SearchType */
+	bool NodePropertyMatchesSearchTokens(const TArray<FString>& Tokens, UBehaviorTreeGraphNode* Node, FString& OutExactFieldValueFound) const;
+
+	/** Determines if a string matches a field property according to the SearchType */
+	bool FieldPropertyMatchesSearchTokens(const TArray<FString>& Tokens, const FProperty* Property, void* Data, FString& OutExactFieldValueFound) const;
+
+	/** Initiates search when SearchType changes */
+	void OnSearchTypeSelectedItemChanged(TSharedPtr<FString> NewValue, ESelectInfo::Type SelectInfo);
+
 private:
 	/** Pointer back to the behavior tree editor that owns us */
 	TWeakPtr<class FBehaviorTreeEditor> BehaviorTreeEditorPtr;
@@ -130,4 +155,10 @@ private:
 
 	/** The string to search for */
 	FString	SearchValue;
+
+	/** Shared strings based on the search types */
+	TArray<TSharedPtr<FString>> SearchTypeComboBoxItems;
+
+	/** Current search type determines how search occurs */
+	EFindInBTSearchType SearchType = SearchType = EFindInBTSearchType::Node;
 };
