Fix static analyzer warning
#rnx

diff --git a/Engine/Source/Editor/BehaviorTreeEditor/Private/FindInBT.h b/Engine/Source/Editor/BehaviorTreeEditor/Private/FindInBT.h
--- a/Engine/Source/Editor/BehaviorTreeEditor/Private/FindInBT.h
+++ b/Engine/Source/Editor/BehaviorTreeEditor/Private/FindInBT.h
@@ -160,5 +160,5 @@ private:
 	TArray<TSharedPtr<FString>> SearchTypeComboBoxItems;
 
 	/** Current search type determines how search occurs */
-	EFindInBTSearchType SearchType = SearchType = EFindInBTSearchType::Node;
+	EFindInBTSearchType SearchType = EFindInBTSearchType::Node;
 };
