Add default value to blackboard key class, enum and object.
Create editor details for the key types to ensure default value are the right type
#rnx

diff --git a/Engine/Source/Editor/BehaviorTreeEditor/BehaviorTreeEditor.Build.cs b/Engine/Source/Editor/BehaviorTreeEditor/BehaviorTreeEditor.Build.cs
--- a/Engine/Source/Editor/BehaviorTreeEditor/BehaviorTreeEditor.Build.cs
+++ b/Engine/Source/Editor/BehaviorTreeEditor/BehaviorTreeEditor.Build.cs
@@ -15,11 +15,13 @@ public class BehaviorTreeEditor : ModuleRules
 				"AssetDefinition",
 				"Core",
 				"CoreUObject",
+				"EditorWidgets",
 				"Engine",
 				"GraphEditor",
 				"GameplayTags",
 				"InputCore",
 				"KismetWidgets",
+				"PropertyEditor",
 				"Slate",
 				"SlateCore",
 				"ToolMenus",
diff --git a/Engine/Source/Editor/BehaviorTreeEditor/Private/DetailCustomizations/BlackboardKeysDetails.cpp b/Engine/Source/Editor/BehaviorTreeEditor/Private/DetailCustomizations/BlackboardKeysDetails.cpp
new file mode 100644
--- /dev/null
+++ b/Engine/Source/Editor/BehaviorTreeEditor/Private/DetailCustomizations/BlackboardKeysDetails.cpp
@@ -0,0 +1,222 @@
+// Copyright Epic Games, Inc. All Rights Reserved.
+
+#include "BlackboardKeysDetails.h"
+
+#include "BehaviorTree/Blackboard/BlackboardKeyType.h"
+#include "UObject/Field.h"
+
+#include "BehaviorTree/Blackboard/BlackboardKeyType_Class.h"
+#include "BehaviorTree/Blackboard/BlackboardKeyType_Enum.h"
+#include "BehaviorTree/Blackboard/BlackboardKeyType_Object.h"
+#include "BehaviorTree/BlackboardData.h"
+#include "Delegates/Delegate.h"
+#include "DetailLayoutBuilder.h"
+#include "Framework/MultiBox/MultiBoxBuilder.h"
+#include "IDetailChildrenBuilder.h"
+#include "IPropertyUtilities.h"
+#include "Internationalization/Internationalization.h"
+#include "Layout/Margin.h"
+#include "Misc/Attribute.h"
+#include "PropertyCustomizationHelpers.h"
+#include "PropertyHandle.h"
+#include "SlateOptMacros.h"
+#include "SEnumCombo.h"
+#include "SlotBase.h"
+#include "Textures/SlateIcon.h"
+#include "Types/SlateEnums.h"
+#include "UObject/Class.h"
+#include "UObject/ObjectPtr.h"
+#include "Widgets/DeclarativeSyntaxSupport.h"
+#include "Widgets/Input/SComboButton.h"
+#include "Widgets/SBoxPanel.h"
+#include "Widgets/SWidget.h"
+#include "Widgets/Text/STextBlock.h"
+
+#define LOCTEXT_NAMESPACE "BlackboardKeysDetails"
+
+TSharedRef<IDetailCustomization> FBlackboardKeyDetails_Class::MakeInstance()
+{
+	return MakeShareable(new FBlackboardKeyDetails_Class());
+}
+
+void FBlackboardKeyDetails_Class::CustomizeDetails(IDetailLayoutBuilder& DetailBuilder)
+{
+	CachedUtils = DetailBuilder.GetPropertyUtilities();
+	BaseClassProperty = DetailBuilder.GetProperty(GET_MEMBER_NAME_CHECKED(UBlackboardKeyType_Class, BaseClass));
+	BaseClassProperty->SetOnPropertyValueChanged(FSimpleDelegate::CreateSP(this, &FBlackboardKeyDetails_Class::OnBaseClassChanged));
+
+	UObject* BaseClassObject = nullptr;
+	BaseClassProperty->GetValue(BaseClassObject);
+	UClass* BaseClass = Cast<UClass>(BaseClassObject);
+
+	DefaultValueProperty = DetailBuilder.GetProperty(GET_MEMBER_NAME_CHECKED(UBlackboardKeyType_Class, DefaultValue));
+	DefaultValueProperty->MarkHiddenByCustomization();
+
+	DetailBuilder.AddCustomRowToCategory(DefaultValueProperty, LOCTEXT("DefaultValue", "DefaultValue"))
+		.NameContent()
+		[
+			DefaultValueProperty->CreatePropertyNameWidget()
+		]
+		.ValueContent()
+		[
+			SNew(SClassPropertyEntryBox)
+				.MetaClass(BaseClass)
+				.AllowNone(true)
+				.AllowAbstract(true)
+				.OnSetClass(this, &FBlackboardKeyDetails_Class::OnSetClass)
+				.SelectedClass(this, &FBlackboardKeyDetails_Class::OnGetSelectedClass)
+		];
+}
+
+void FBlackboardKeyDetails_Class::OnBaseClassChanged()
+{
+	UObject* BaseClassObject = nullptr;
+	BaseClassProperty->GetValue(BaseClassObject);
+	const UClass* BaseClass = Cast<UClass>(BaseClassObject);
+	const UClass* SelectedClass = OnGetSelectedClass();
+	if (SelectedClass && SelectedClass != BaseClass && !SelectedClass->IsChildOf(BaseClass))
+	{
+		OnSetClass(nullptr);
+	}
+
+	if (CachedUtils)
+	{
+		CachedUtils->ForceRefresh();
+	}
+}
+
+void FBlackboardKeyDetails_Class::OnSetClass(const UClass* NewClass)
+{
+	DefaultValueProperty->SetValue(NewClass);
+}
+
+const UClass* FBlackboardKeyDetails_Class::OnGetSelectedClass() const
+{
+	UObject* DefaultValue = nullptr;
+	DefaultValueProperty->GetValue(DefaultValue);
+	return Cast<UClass>(DefaultValue);
+}
+
+TSharedRef<IDetailCustomization> FBlackboardKeyDetails_Object::MakeInstance()
+{
+	return MakeShareable(new FBlackboardKeyDetails_Object());
+}
+
+void FBlackboardKeyDetails_Object::CustomizeDetails(IDetailLayoutBuilder& DetailBuilder)
+{
+	CachedUtils = DetailBuilder.GetPropertyUtilities();
+	BaseClassProperty = DetailBuilder.GetProperty(GET_MEMBER_NAME_CHECKED(UBlackboardKeyType_Object, BaseClass));
+	BaseClassProperty->SetOnPropertyValueChanged(FSimpleDelegate::CreateSP(this, &FBlackboardKeyDetails_Object::OnBaseClassChanged));
+
+	UObject* BaseClassObject = nullptr;
+	BaseClassProperty->GetValue(BaseClassObject);
+	UClass* BaseClass = Cast<UClass>(BaseClassObject);
+
+	DefaultValueProperty = DetailBuilder.GetProperty(GET_MEMBER_NAME_CHECKED(UBlackboardKeyType_Object, DefaultValue));
+	DefaultValueProperty->MarkHiddenByCustomization();
+
+	DetailBuilder.AddCustomRowToCategory(DefaultValueProperty, LOCTEXT("DefaultValue", "DefaultValue"))
+		.NameContent()
+		[
+			DefaultValueProperty->CreatePropertyNameWidget()
+		]
+		.ValueContent()
+		[
+			SNew(SObjectPropertyEntryBox)
+				.AllowedClass(BaseClass)
+				.AllowClear(true)
+				.OnObjectChanged(this, &FBlackboardKeyDetails_Object::OnObjectChanged)
+				.ObjectPath(this, &FBlackboardKeyDetails_Object::OnGetObjectPath)
+		];
+}
+
+void FBlackboardKeyDetails_Object::OnBaseClassChanged()
+{
+	UObject* BaseClassObject = nullptr;
+	BaseClassProperty->GetValue(BaseClassObject);
+	const UClass* BaseClass = Cast<UClass>(BaseClassObject);
+	const UObject* SelectedObject = OnGetSelectedObject();
+	if (SelectedObject && SelectedObject->GetClass() != BaseClass && !SelectedObject->GetClass()->IsChildOf(BaseClass))
+	{
+		DefaultValueProperty->SetValue(static_cast<UObject*>(nullptr));
+	}
+
+	if (CachedUtils)
+	{
+		CachedUtils->ForceRefresh();
+	}
+}
+
+const UObject* FBlackboardKeyDetails_Object::OnGetSelectedObject() const
+{
+	UObject* DefaultValue = nullptr;
+	DefaultValueProperty->GetValue(DefaultValue);
+	return DefaultValue;
+}
+
+void FBlackboardKeyDetails_Object::OnObjectChanged(const FAssetData& AssetData)
+{
+	DefaultValueProperty->SetValue(AssetData.GetAsset());
+}
+
+FString FBlackboardKeyDetails_Object::OnGetObjectPath() const
+{
+	const UObject* DefaultValue = OnGetSelectedObject();
+	return DefaultValue ? DefaultValue->GetPathName() : FString();
+}
+
+TSharedRef<IDetailCustomization> FBlackboardKeyDetails_Enum::MakeInstance()
+{
+	return MakeShareable(new FBlackboardKeyDetails_Enum());
+}
+
+void FBlackboardKeyDetails_Enum::CustomizeDetails(IDetailLayoutBuilder& DetailBuilder)
+{
+	CachedUtils = DetailBuilder.GetPropertyUtilities();
+	EnumTypeProperty = DetailBuilder.GetProperty(GET_MEMBER_NAME_CHECKED(UBlackboardKeyType_Enum, EnumType));
+	EnumTypeProperty->SetOnPropertyValueChanged(FSimpleDelegate::CreateSP(this, &FBlackboardKeyDetails_Enum::OnEnumTypeChanged));
+
+	UObject* EnumTypeObject = nullptr;
+	EnumTypeProperty->GetValue(EnumTypeObject);
+	UEnum* EnumType = Cast<UEnum>(EnumTypeObject);
+
+	DefaultValueProperty = DetailBuilder.GetProperty(GET_MEMBER_NAME_CHECKED(UBlackboardKeyType_Enum, DefaultValue));
+	DefaultValueProperty->MarkHiddenByCustomization();
+
+	if(EnumType)
+	{
+		DetailBuilder.AddCustomRowToCategory(DefaultValueProperty, LOCTEXT("DefaultValue", "DefaultValue"))
+			.NameContent()
+			[
+				DefaultValueProperty->CreatePropertyNameWidget()
+			]
+			.ValueContent()
+			[
+				SNew(SEnumComboBox, EnumType)
+					.CurrentValue_Raw(this, &FBlackboardKeyDetails_Enum::GetEnumValue)
+					.OnEnumSelectionChanged(SEnumComboBox::FOnEnumSelectionChanged::CreateSP(this, &FBlackboardKeyDetails_Enum::OnEnumSelectionChanged))
+			];
+	}
+}
+
+void FBlackboardKeyDetails_Enum::OnEnumSelectionChanged(int32 NewValue, ESelectInfo::Type)
+{
+	DefaultValueProperty->SetValue(static_cast<uint8>(NewValue));
+}
+
+void FBlackboardKeyDetails_Enum::OnEnumTypeChanged()
+{
+	DefaultValueProperty->SetValue(static_cast<uint8>(0));
+	if (CachedUtils)
+	{
+		CachedUtils->ForceRefresh();
+	}
+}
+
+int32 FBlackboardKeyDetails_Enum::GetEnumValue() const
+{
+	uint8 EnumValue = 0;
+	DefaultValueProperty->GetValue(EnumValue);
+	return EnumValue;
+}
+#undef LOCTEXT_NAMESPACE
diff --git a/Engine/Source/Editor/BehaviorTreeEditor/Private/DetailCustomizations/BlackboardKeysDetails.h b/Engine/Source/Editor/BehaviorTreeEditor/Private/DetailCustomizations/BlackboardKeysDetails.h
new file mode 100644
--- /dev/null
+++ b/Engine/Source/Editor/BehaviorTreeEditor/Private/DetailCustomizations/BlackboardKeysDetails.h
@@ -0,0 +1,75 @@
+// Copyright Epic Games, Inc.All Rights Reserved.
+
+#pragma once
+
+#include "Containers/Array.h"
+#include "HAL/Platform.h"
+#include "IDetailCustomization.h"
+#include "Internationalization/Text.h"
+#include "Templates/SharedPointer.h"
+#include "UObject/NameTypes.h"
+
+class IPropertyHandle;
+class IPropertyUtilities;
+class SWidget;
+struct FAssetData;
+
+namespace ESelectInfo { enum Type : int; }
+
+class FBlackboardKeyDetails_Class : public IDetailCustomization
+{
+public:
+	/** Makes a new instance of this detail layout class for a specific detail view requesting it */
+	static TSharedRef<IDetailCustomization> MakeInstance();
+
+	/** IDetailCustomization interface */
+	virtual void CustomizeDetails(IDetailLayoutBuilder& DetailBuilder) override;
+
+protected:
+	TSharedPtr<IPropertyHandle> BaseClassProperty;
+	TSharedPtr<IPropertyHandle> DefaultValueProperty;
+	TSharedPtr<IPropertyUtilities> CachedUtils;
+
+	void OnBaseClassChanged();
+	void OnSetClass(const UClass* NewClass);
+	const UClass* OnGetSelectedClass() const;
+};
+
+class FBlackboardKeyDetails_Object : public IDetailCustomization
+{
+public:
+	/** Makes a new instance of this detail layout class for a specific detail view requesting it */
+	static TSharedRef<IDetailCustomization> MakeInstance();
+
+	/** IDetailCustomization interface */
+	virtual void CustomizeDetails(IDetailLayoutBuilder& DetailBuilder) override;
+
+protected:
+	TSharedPtr<IPropertyHandle> BaseClassProperty;
+	TSharedPtr<IPropertyHandle> DefaultValueProperty;
+	TSharedPtr<IPropertyUtilities> CachedUtils;
+
+	void OnBaseClassChanged();
+	const UObject* OnGetSelectedObject() const;
+	void OnObjectChanged(const FAssetData& AssetData);
+	FString OnGetObjectPath() const;
+};
+
+class FBlackboardKeyDetails_Enum : public IDetailCustomization
+{
+public:
+	/** Makes a new instance of this detail layout class for a specific detail view requesting it */
+	static TSharedRef<IDetailCustomization> MakeInstance();
+
+	/** IDetailCustomization interface */
+	virtual void CustomizeDetails(IDetailLayoutBuilder& DetailBuilder) override;
+
+protected:
+	TSharedPtr<IPropertyHandle> EnumTypeProperty;
+	TSharedPtr<IPropertyHandle> DefaultValueProperty;
+	TSharedPtr<IPropertyUtilities> CachedUtils;
+
+	void OnEnumSelectionChanged(int32 NewValue, ESelectInfo::Type);
+	void OnEnumTypeChanged();
+	int32 GetEnumValue() const;
+};
\ No newline at end of file
diff --git a/Engine/Source/Runtime/AIModule/Classes/BehaviorTree/Blackboard/BlackboardKeyType_Class.h b/Engine/Source/Runtime/AIModule/Classes/BehaviorTree/Blackboard/BlackboardKeyType_Class.h
--- a/Engine/Source/Runtime/AIModule/Classes/BehaviorTree/Blackboard/BlackboardKeyType_Class.h
+++ b/Engine/Source/Runtime/AIModule/Classes/BehaviorTree/Blackboard/BlackboardKeyType_Class.h
@@ -16,9 +16,12 @@ class UBlackboardKeyType_Class : public UBlackboardKeyType
 	typedef UClass* FDataType;
 	static AIMODULE_API const FDataType InvalidValue;
 
-	UPROPERTY(Category=Blackboard, EditDefaultsOnly, meta=(AllowAbstract="1"))
+	UPROPERTY(Category=Blackboard, EditDefaultsOnly, NoClear, meta=(AllowAbstract="1"))
 	TObjectPtr<UClass> BaseClass;
 	
+	UPROPERTY(Category=Blackboard, EditDefaultsOnly)
+	TObjectPtr<UClass> DefaultValue;
+
 	static AIMODULE_API UClass* GetValue(const UBlackboardKeyType_Class* KeyOb, const uint8* RawData);
 	static AIMODULE_API bool SetValue(UBlackboardKeyType_Class* KeyOb, uint8* RawData, UClass* Value);
 
@@ -28,6 +31,8 @@ class UBlackboardKeyType_Class : public UBlackboardKeyType
 	AIMODULE_API virtual FString DescribeSelf() const override;
 	AIMODULE_API virtual bool IsAllowedByFilter(UBlackboardKeyType* FilterOb) const override;
 
+	AIMODULE_API virtual void InitializeMemory(UBlackboardComponent& OwnerComp, uint8* MemoryBlock) override;
+
 protected:
 	AIMODULE_API virtual FString DescribeValue(const UBlackboardComponent& OwnerComp, const uint8* RawData) const override;
 	AIMODULE_API virtual bool TestBasicOperation(const UBlackboardComponent& OwnerComp, const uint8* MemoryBlock, EBasicKeyOperation::Type Op) const override;
diff --git a/Engine/Source/Runtime/AIModule/Classes/BehaviorTree/Blackboard/BlackboardKeyType_Enum.h b/Engine/Source/Runtime/AIModule/Classes/BehaviorTree/Blackboard/BlackboardKeyType_Enum.h
--- a/Engine/Source/Runtime/AIModule/Classes/BehaviorTree/Blackboard/BlackboardKeyType_Enum.h
+++ b/Engine/Source/Runtime/AIModule/Classes/BehaviorTree/Blackboard/BlackboardKeyType_Enum.h
@@ -23,6 +23,9 @@ class UBlackboardKeyType_Enum : public UBlackboardKeyType
 	UPROPERTY(Category=Blackboard, EditDefaultsOnly)
 	FString EnumName;
 
+	UPROPERTY(Category=Blackboard, EditDefaultsOnly)
+	uint8 DefaultValue = 0;
+
 	/** set when EnumName override is valid and active */
 	UPROPERTY(Category = Blackboard, VisibleDefaultsOnly)
 	uint32 bIsEnumNameValid : 1;
@@ -37,6 +40,8 @@ class UBlackboardKeyType_Enum : public UBlackboardKeyType
 	AIMODULE_API virtual FString DescribeArithmeticParam(int32 IntValue, float FloatValue) const override;
 	AIMODULE_API virtual bool IsAllowedByFilter(UBlackboardKeyType* FilterOb) const override;
 
+	AIMODULE_API virtual void InitializeMemory(UBlackboardComponent& OwnerComp, uint8* MemoryBlock) override;
+
 protected:
 #if WITH_EDITOR
 	AIMODULE_API virtual void PostEditChangeProperty(struct FPropertyChangedEvent& PropertyChangedEvent) override;
diff --git a/Engine/Source/Runtime/AIModule/Classes/BehaviorTree/Blackboard/BlackboardKeyType_Object.h b/Engine/Source/Runtime/AIModule/Classes/BehaviorTree/Blackboard/BlackboardKeyType_Object.h
--- a/Engine/Source/Runtime/AIModule/Classes/BehaviorTree/Blackboard/BlackboardKeyType_Object.h
+++ b/Engine/Source/Runtime/AIModule/Classes/BehaviorTree/Blackboard/BlackboardKeyType_Object.h
@@ -19,6 +19,9 @@ class UBlackboardKeyType_Object : public UBlackboardKeyType
 	UPROPERTY(Category=Blackboard, EditDefaultsOnly, NoClear, meta=(AllowAbstract="1"))
 	TObjectPtr<UClass> BaseClass;
 
+	UPROPERTY(Category=Blackboard, EditDefaultsOnly)
+	TObjectPtr<UObject> DefaultValue;
+
 	static AIMODULE_API UObject* GetValue(const UBlackboardKeyType_Object* KeyOb, const uint8* RawData);
 	static AIMODULE_API bool SetValue(UBlackboardKeyType_Object* KeyOb, uint8* RawData, UObject* Value);
 
@@ -28,6 +31,8 @@ class UBlackboardKeyType_Object : public UBlackboardKeyType
 	AIMODULE_API virtual FString DescribeSelf() const override;
 	AIMODULE_API virtual bool IsAllowedByFilter(UBlackboardKeyType* FilterOb) const override;
 
+	AIMODULE_API virtual void InitializeMemory(UBlackboardComponent& OwnerComp, uint8* MemoryBlock) override;
+
 protected:
 	AIMODULE_API virtual FString DescribeValue(const UBlackboardComponent& OwnerComp, const uint8* RawData) const override;
 	AIMODULE_API virtual bool GetLocation(const UBlackboardComponent& OwnerComp, const uint8* RawData, FVector& Location) const override;
diff --git a/Engine/Source/Runtime/AIModule/Classes/BehaviorTree/BlackboardData.h b/Engine/Source/Runtime/AIModule/Classes/BehaviorTree/BlackboardData.h
--- a/Engine/Source/Runtime/AIModule/Classes/BehaviorTree/BlackboardData.h
+++ b/Engine/Source/Runtime/AIModule/Classes/BehaviorTree/BlackboardData.h
@@ -27,7 +27,7 @@ struct FBlackboardEntry
 #endif // WITH_EDITORONLY_DATA
 
 	/** key type and additional properties */
-	UPROPERTY(EditAnywhere, Instanced, Category=Blackboard)
+	UPROPERTY(EditAnywhere, Instanced, Category = Blackboard, meta = (AllowEditInlineCustomization))
 	TObjectPtr<UBlackboardKeyType> KeyType;
 
 	/** if set to true then this field will be synchronized across all instances of this blackboard */
diff --git a/Engine/Source/Runtime/AIModule/Private/BehaviorTree/Blackboard/BlackboardKeyType_Class.cpp b/Engine/Source/Runtime/AIModule/Private/BehaviorTree/Blackboard/BlackboardKeyType_Class.cpp
--- a/Engine/Source/Runtime/AIModule/Private/BehaviorTree/Blackboard/BlackboardKeyType_Class.cpp
+++ b/Engine/Source/Runtime/AIModule/Private/BehaviorTree/Blackboard/BlackboardKeyType_Class.cpp
@@ -36,6 +36,11 @@ EBlackboardCompare::Type UBlackboardKeyType_Class::CompareValues(const UBlackboa
 	return (MyValue == OtherValue) ? EBlackboardCompare::Equal : EBlackboardCompare::NotEqual;
 }
 
+void UBlackboardKeyType_Class::InitializeMemory(UBlackboardComponent& OwnerComp, uint8* MemoryBlock)
+{
+	SetValue(this, MemoryBlock, DefaultValue);
+}
+
 FString UBlackboardKeyType_Class::DescribeValue(const UBlackboardComponent& OwnerComp, const uint8* RawData) const
 {
 	return *GetNameSafe(GetValue(this, RawData));
diff --git a/Engine/Source/Runtime/AIModule/Private/BehaviorTree/Blackboard/BlackboardKeyType_Enum.cpp b/Engine/Source/Runtime/AIModule/Private/BehaviorTree/Blackboard/BlackboardKeyType_Enum.cpp
--- a/Engine/Source/Runtime/AIModule/Private/BehaviorTree/Blackboard/BlackboardKeyType_Enum.cpp
+++ b/Engine/Source/Runtime/AIModule/Private/BehaviorTree/Blackboard/BlackboardKeyType_Enum.cpp
@@ -78,6 +78,11 @@ FString UBlackboardKeyType_Enum::DescribeArithmeticParam(int32 IntValue, float F
 	return EnumType ? EnumType->GetDisplayNameTextByValue(IntValue).ToString() : FString("UNKNOWN!");
 }
 
+void UBlackboardKeyType_Enum::InitializeMemory(UBlackboardComponent& OwnerComp, uint8* MemoryBlock)
+{
+	SetValue(this, MemoryBlock, DefaultValue);
+}
+
 #if WITH_EDITOR
 void UBlackboardKeyType_Enum::PostEditChangeProperty(struct FPropertyChangedEvent& PropertyChangedEvent)
 {
diff --git a/Engine/Source/Runtime/AIModule/Private/BehaviorTree/Blackboard/BlackboardKeyType_Object.cpp b/Engine/Source/Runtime/AIModule/Private/BehaviorTree/Blackboard/BlackboardKeyType_Object.cpp
--- a/Engine/Source/Runtime/AIModule/Private/BehaviorTree/Blackboard/BlackboardKeyType_Object.cpp
+++ b/Engine/Source/Runtime/AIModule/Private/BehaviorTree/Blackboard/BlackboardKeyType_Object.cpp
@@ -97,3 +97,7 @@ bool UBlackboardKeyType_Object::TestBasicOperation(const UBlackboardComponent& O
 	return (Op == EBasicKeyOperation::Set) ? WeakObjPtr.IsValid() : !WeakObjPtr.IsValid();
 }
 
+void UBlackboardKeyType_Object::InitializeMemory(UBlackboardComponent& OwnerComp, uint8* MemoryBlock)
+{
+	SetValue(this, MemoryBlock, DefaultValue);
+}
