Missing file in cl 34071056
#rnx

diff --git a/Engine/Source/Editor/BehaviorTreeEditor/Private/BehaviorTreeEditorModule.cpp b/Engine/Source/Editor/BehaviorTreeEditor/Private/BehaviorTreeEditorModule.cpp
--- a/Engine/Source/Editor/BehaviorTreeEditor/Private/BehaviorTreeEditorModule.cpp
+++ b/Engine/Source/Editor/BehaviorTreeEditor/Private/BehaviorTreeEditorModule.cpp
@@ -5,6 +5,7 @@
 #include "BehaviorTreeDecoratorGraphNode_Decorator.h"
 #include "BehaviorTreeGraphNode.h"
 #include "EdGraphUtilities.h"
+#include "DetailCustomizations/BlackboardKeysDetails.h"
 #include "DetailCustomizations/BlackboardSelectorDetails.h"
 #include "BehaviorTreeEditor.h"
 #include "SGraphNode_BehaviorTree.h"
@@ -54,6 +55,9 @@ void FBehaviorTreeEditorModule::StartupModule()
 	PropertyModule.RegisterCustomPropertyTypeLayout( "BlackboardKeySelector", FOnGetPropertyTypeCustomizationInstance::CreateStatic( &FBlackboardSelectorDetails::MakeInstance ) );
 	PropertyModule.RegisterCustomClassLayout( "BTDecorator_Blackboard", FOnGetDetailCustomizationInstance::CreateStatic( &FBlackboardDecoratorDetails::MakeInstance ) );
 	PropertyModule.RegisterCustomClassLayout( "BTDecorator", FOnGetDetailCustomizationInstance::CreateStatic( &FBehaviorDecoratorDetails::MakeInstance ) );
+	PropertyModule.RegisterCustomClassLayout("BlackboardKeyType_Class", FOnGetDetailCustomizationInstance::CreateStatic(&FBlackboardKeyDetails_Class::MakeInstance));
+	PropertyModule.RegisterCustomClassLayout("BlackboardKeyType_Enum", FOnGetDetailCustomizationInstance::CreateStatic(&FBlackboardKeyDetails_Enum::MakeInstance));
+	PropertyModule.RegisterCustomClassLayout("BlackboardKeyType_Object", FOnGetDetailCustomizationInstance::CreateStatic(&FBlackboardKeyDetails_Object::MakeInstance));
 	PropertyModule.NotifyCustomizationModuleChanged();
 }
 
@@ -82,6 +86,9 @@ void FBehaviorTreeEditorModule::ShutdownModule()
 		PropertyModule.UnregisterCustomPropertyTypeLayout( "BlackboardKeySelector" );
 		PropertyModule.UnregisterCustomClassLayout( "BTDecorator_Blackboard" );
 		PropertyModule.UnregisterCustomClassLayout( "BTDecorator" );
+		PropertyModule.UnregisterCustomClassLayout("BlackboardKeyType_Class");
+		PropertyModule.UnregisterCustomClassLayout("BlackboardKeyType_Enum");
+		PropertyModule.UnregisterCustomClassLayout("BlackboardKeyType_Object");
 		PropertyModule.NotifyCustomizationModuleChanged();
 	}
 }
