Introduce FValueOrBBKey_X, a Default Value/Blackboard Key pair for BTNodes with a custom view to make picking key easier
FValueOrKey are created of each of the engine blackboard key types.

diff --git a/Engine/Plugins/Experimental/GameplayBehaviors/Source/GameplayBehaviorsEditorModule/GameplayBehaviorsEditorModule.Build.cs b/Engine/Plugins/Experimental/GameplayBehaviors/Source/GameplayBehaviorsEditorModule/GameplayBehaviorsEditorModule.Build.cs
--- a/Engine/Plugins/Experimental/GameplayBehaviors/Source/GameplayBehaviorsEditorModule/GameplayBehaviorsEditorModule.Build.cs
+++ b/Engine/Plugins/Experimental/GameplayBehaviors/Source/GameplayBehaviorsEditorModule/GameplayBehaviorsEditorModule.Build.cs
@@ -23,6 +23,7 @@ namespace UnrealBuildTool.Rules
 
 			PrivateDependencyModuleNames.AddRange(
 			new string[] {
+				"BehaviorTreeEditor"
 			}
 			);
 		}
diff --git a/Engine/Plugins/Experimental/GameplayBehaviors/Source/GameplayBehaviorsEditorModule/Private/GameplayBehaviorsEditorModule.cpp b/Engine/Plugins/Experimental/GameplayBehaviors/Source/GameplayBehaviorsEditorModule/Private/GameplayBehaviorsEditorModule.cpp
--- a/Engine/Plugins/Experimental/GameplayBehaviors/Source/GameplayBehaviorsEditorModule/Private/GameplayBehaviorsEditorModule.cpp
+++ b/Engine/Plugins/Experimental/GameplayBehaviors/Source/GameplayBehaviorsEditorModule/Private/GameplayBehaviorsEditorModule.cpp
@@ -4,6 +4,10 @@
 
 #include "GameplayBehaviorsEditorStyle.h"
 
+#include "Modules/ModuleManager.h"
+#include "PropertyEditorModule.h"
+#include "ValueOrBBKeyDetails.h"
+
 #define LOCTEXT_NAMESPACE "GameplayBehaviors"
 
 class FGameplayBehaviorsEditorModule : public IGameplayBehaviorsEditorModule
@@ -16,11 +20,19 @@ protected:
 void FGameplayBehaviorsEditorModule::StartupModule()
 {
 	FGameplayBehaviorsEditorStyle::Get();
+
+	FPropertyEditorModule& PropertyModule = FModuleManager::LoadModuleChecked<FPropertyEditorModule>("PropertyEditor");
+	PropertyModule.RegisterCustomPropertyTypeLayout("ValueOrBBKey_GameplayTag", FOnGetPropertyTypeCustomizationInstance::CreateStatic(&FValueOrBBKeyDetails::MakeInstance));
+	PropertyModule.NotifyCustomizationModuleChanged();
 }
 
 void FGameplayBehaviorsEditorModule::ShutdownModule()
 {
 	FGameplayBehaviorsEditorStyle::Shutdown();
+
+	FPropertyEditorModule& PropertyModule = FModuleManager::LoadModuleChecked<FPropertyEditorModule>("PropertyEditor");
+	PropertyModule.UnregisterCustomPropertyTypeLayout("ValueOrBBKey_GameplayTag");
+	PropertyModule.NotifyCustomizationModuleChanged();
 }
 
 IMPLEMENT_MODULE(FGameplayBehaviorsEditorModule, GameplayBehaviorsEditorModule)
diff --git a/Engine/Plugins/Experimental/GameplayBehaviors/Source/GameplayBehaviorsModule/Private/AI/ValueOrBBKey_GameplayTag.cpp b/Engine/Plugins/Experimental/GameplayBehaviors/Source/GameplayBehaviorsModule/Private/AI/ValueOrBBKey_GameplayTag.cpp
new file mode 100644
--- /dev/null
+++ b/Engine/Plugins/Experimental/GameplayBehaviors/Source/GameplayBehaviorsModule/Private/AI/ValueOrBBKey_GameplayTag.cpp
@@ -0,0 +1,50 @@
+// Copyright Epic Games, Inc. All Rights Reserved.
+#pragma once
+
+#include "AI/ValueOrBBKey_GameplayTag.h"
+
+#include "BlackboardKeyType_GameplayTag.h"
+
+FString FValueOrBBKey_GameplayTagContainer::ToString() const
+{
+	if (!Key.IsNone())
+	{
+		return ToStringKeyName();
+	}
+	else
+	{
+		return DefaultValue.ToStringSimple(false);
+	}
+}
+
+FGameplayTagContainer FValueOrBBKey_GameplayTagContainer::GetValue(const UBlackboardComponent& Blackboard) const
+{
+	return FBlackboard::GetValue<UBlackboardKeyType_GameplayTag>(Blackboard, Key, KeyId, DefaultValue);
+}
+
+FGameplayTagContainer FValueOrBBKey_GameplayTagContainer::GetValue(const UBlackboardComponent* Blackboard) const
+{
+	return Blackboard ? GetValue(*Blackboard) : DefaultValue;
+}
+
+FGameplayTagContainer FValueOrBBKey_GameplayTagContainer::GetValue(const UBehaviorTreeComponent& BehaviorComp) const
+{
+	return FBlackboard::GetValue<UBlackboardKeyType_GameplayTag>(BehaviorComp, Key, KeyId, DefaultValue);
+}
+
+FGameplayTagContainer FValueOrBBKey_GameplayTagContainer::GetValue(const UBehaviorTreeComponent* BehaviorComp) const
+{
+	return BehaviorComp ? GetValue(*BehaviorComp) : DefaultValue;
+}
+
+FGameplayTagContainer UValueOrBBKey_GameplayTagBlueprintUtility::GetTagContainer(const FValueOrBBKey_GameplayTagContainer& Value, const UBehaviorTreeComponent* BehaviorTreeComp)
+{
+	return Value.GetValue(BehaviorTreeComp);
+}
+
+#if WITH_EDITOR
+bool FValueOrBBKey_GameplayTagContainer::IsCompatibleType(const UBlackboardKeyType* KeyType) const
+{
+	return KeyType && KeyType->GetClass() == UBlackboardKeyType_GameplayTag::StaticClass();
+}
+#endif // WITH_EDITOR
diff --git a/Engine/Plugins/Experimental/GameplayBehaviors/Source/GameplayBehaviorsModule/Public/AI/ValueOrBBKey_GameplayTag.h b/Engine/Plugins/Experimental/GameplayBehaviors/Source/GameplayBehaviorsModule/Public/AI/ValueOrBBKey_GameplayTag.h
new file mode 100644
--- /dev/null
+++ b/Engine/Plugins/Experimental/GameplayBehaviors/Source/GameplayBehaviorsModule/Public/AI/ValueOrBBKey_GameplayTag.h
@@ -0,0 +1,40 @@
+// Copyright Epic Games, Inc. All Rights Reserved.
+
+#pragma once
+
+#include "BehaviorTree/ValueOrBBKey.h"
+#include "Kismet/BlueprintFunctionLibrary.h"
+
+#include "ValueOrBBKey_GameplayTag.generated.h"
+
+USTRUCT(BlueprintType)
+struct FValueOrBBKey_GameplayTagContainer : public FValueOrBlackboardKeyBase
+{
+	GENERATED_BODY()
+
+	FValueOrBBKey_GameplayTagContainer(FGameplayTagContainer Default = FGameplayTagContainer())
+		: DefaultValue(MoveTemp(Default)) {}
+	GAMEPLAYBEHAVIORSMODULE_API FGameplayTagContainer GetValue(const UBehaviorTreeComponent& BehaviorComp) const;
+	GAMEPLAYBEHAVIORSMODULE_API FGameplayTagContainer GetValue(const UBehaviorTreeComponent* BehaviorComp) const;
+	GAMEPLAYBEHAVIORSMODULE_API FGameplayTagContainer GetValue(const UBlackboardComponent& Blackboard) const;
+	GAMEPLAYBEHAVIORSMODULE_API FGameplayTagContainer GetValue(const UBlackboardComponent* Blackboard) const;
+
+#if WITH_EDITOR
+	GAMEPLAYBEHAVIORSMODULE_API virtual bool IsCompatibleType(const UBlackboardKeyType* KeyType) const override;
+#endif // WITH_EDITOR
+
+	GAMEPLAYBEHAVIORSMODULE_API FString ToString() const;
+
+protected:
+	UPROPERTY(EditAnywhere, Category = "Value")
+	FGameplayTagContainer DefaultValue;
+};
+
+UCLASS()
+class UValueOrBBKey_GameplayTagBlueprintUtility : public UBlueprintFunctionLibrary
+{
+	GENERATED_BODY()
+public:
+	UFUNCTION(BlueprintPure, Category = Blackboard)
+	static FGameplayTagContainer GetTagContainer(const FValueOrBBKey_GameplayTagContainer& Value, const UBehaviorTreeComponent* BehaviorTreeComp);
+};
\ No newline at end of file
diff --git a/Engine/Source/Editor/BehaviorTreeEditor/Classes/ValueOrBBKeyDetails.h b/Engine/Source/Editor/BehaviorTreeEditor/Classes/ValueOrBBKeyDetails.h
new file mode 100644
--- /dev/null
+++ b/Engine/Source/Editor/BehaviorTreeEditor/Classes/ValueOrBBKeyDetails.h
@@ -0,0 +1,113 @@
+// Copyright Epic Games, Inc.All Rights Reserved.
+
+#pragma once
+
+#include "Containers/Array.h"
+#include "HAL/Platform.h"
+#include "IPropertyTypeCustomization.h"
+#include "Internationalization/Text.h"
+#include "Templates/SharedPointer.h"
+#include "UObject/NameTypes.h"
+
+namespace ESelectInfo
+{
+	enum Type : int;
+}
+
+class IPropertyHandle;
+class IPropertyUtilities;
+class SWidget;
+class UBlackboardData;
+struct FValueOrBlackboardKeyBase;
+
+class BEHAVIORTREEEDITOR_API FValueOrBBKeyDetails : public IPropertyTypeCustomization
+{
+public:
+	/** Makes a new instance of this detail layout class for a specific detail view requesting it */
+	static TSharedRef<IPropertyTypeCustomization> MakeInstance();
+
+	virtual void CustomizeHeader(TSharedRef<class IPropertyHandle> StructPropertyHandle, class FDetailWidgetRow& HeaderRow, IPropertyTypeCustomizationUtils& StructCustomizationUtils) override;
+	virtual void CustomizeChildren(TSharedRef<class IPropertyHandle> StructPropertyHandle, class IDetailChildrenBuilder& StructBuilder, IPropertyTypeCustomizationUtils& StructCustomizationUtils) override;
+
+	bool CanEditDefaultValue() const;
+
+protected:
+	TSharedPtr<IPropertyHandle> StructProperty;
+	TSharedPtr<IPropertyHandle> KeyProperty;
+	TSharedPtr<IPropertyHandle> DefaultValueProperty;
+	TSharedPtr<IPropertyUtilities> CachedUtils;
+	TArray<FName> MatchingKeys;
+
+	virtual void ValidateData();
+	virtual TSharedRef<SWidget> CreateDefaultValueWidget();
+
+	void GetMatchingKeys(TArray<FName>& OutNames);
+	bool HasAccessToBlackboard() const;
+	TSharedRef<SWidget> OnGetKeyNames();
+	void OnKeyChanged(int32 Index);
+	FText GetKeyDesc() const;
+	const FValueOrBlackboardKeyBase* GetDataPtr() const;
+};
+
+class FValueOrBBKeyDetails_Class : public FValueOrBBKeyDetails
+{
+public:
+	/** Makes a new instance of this detail layout class for a specific detail view requesting it */
+	static TSharedRef<IPropertyTypeCustomization> MakeInstance();
+
+	virtual void CustomizeHeader(TSharedRef<class IPropertyHandle> StructPropertyHandle, class FDetailWidgetRow& HeaderRow, IPropertyTypeCustomizationUtils& StructCustomizationUtils) override;
+	virtual void CustomizeChildren(TSharedRef<class IPropertyHandle> StructPropertyHandle, class IDetailChildrenBuilder& StructBuilder, IPropertyTypeCustomizationUtils& StructCustomizationUtils) override;
+
+protected:
+	TSharedPtr<IPropertyHandle> BaseClassProperty;
+
+	virtual void ValidateData() override;
+	virtual TSharedRef<SWidget> CreateDefaultValueWidget() override;
+
+	void OnBaseClassChanged();
+	void OnSetClass(const UClass* NewClass);
+	const UClass* OnGetSelectedClass() const;
+};
+
+class FValueOrBBKeyDetails_Enum : public FValueOrBBKeyDetails
+{
+public:
+	/** Makes a new instance of this detail layout class for a specific detail view requesting it */
+	static TSharedRef<IPropertyTypeCustomization> MakeInstance();
+
+	virtual void CustomizeHeader(TSharedRef<class IPropertyHandle> StructPropertyHandle, class FDetailWidgetRow& HeaderRow, IPropertyTypeCustomizationUtils& StructCustomizationUtils) override;
+	virtual void CustomizeChildren(TSharedRef<class IPropertyHandle> StructPropertyHandle, class IDetailChildrenBuilder& StructBuilder, IPropertyTypeCustomizationUtils& StructCustomizationUtils) override;
+
+protected:
+	TSharedPtr<IPropertyHandle> EnumTypeProperty;
+	TSharedPtr<IPropertyHandle> NativeEnumTypeNameProperty;
+
+	virtual void ValidateData() override;
+	virtual TSharedRef<SWidget> CreateDefaultValueWidget() override;
+
+	void OnEnumSelectionChanged(int32 NewValue, ESelectInfo::Type);
+	void OnEnumTypeChanged();
+	void OnNativeEnumTypeNameChanged();
+	int32 GetEnumValue() const;
+	bool CanEditEnumType() const;
+};
+
+class FValueOrBBKeyDetails_Object : public FValueOrBBKeyDetails
+{
+public:
+	/** Makes a new instance of this detail layout class for a specific detail view requesting it */
+	static TSharedRef<IPropertyTypeCustomization> MakeInstance();
+
+	virtual void CustomizeHeader(TSharedRef<class IPropertyHandle> StructPropertyHandle, class FDetailWidgetRow& HeaderRow, IPropertyTypeCustomizationUtils& StructCustomizationUtils) override;
+	virtual void CustomizeChildren(TSharedRef<class IPropertyHandle> StructPropertyHandle, class IDetailChildrenBuilder& StructBuilder, IPropertyTypeCustomizationUtils& StructCustomizationUtils) override;
+
+protected:
+	TSharedPtr<IPropertyHandle> BaseClassProperty;
+
+	virtual void ValidateData() override;
+	virtual TSharedRef<SWidget> CreateDefaultValueWidget() override;
+
+	void OnBaseClassChanged();
+	void OnObjectChanged(const FAssetData& AssetData);
+	FString OnGetObjectPath() const;
+};
diff --git a/Engine/Source/Editor/BehaviorTreeEditor/Private/BehaviorTreeEditorModule.cpp b/Engine/Source/Editor/BehaviorTreeEditor/Private/BehaviorTreeEditorModule.cpp
--- a/Engine/Source/Editor/BehaviorTreeEditor/Private/BehaviorTreeEditorModule.cpp
+++ b/Engine/Source/Editor/BehaviorTreeEditor/Private/BehaviorTreeEditorModule.cpp
@@ -2,25 +2,27 @@
 
 #include "BehaviorTreeEditorModule.h"
 
-#include "BehaviorTreeDecoratorGraphNode_Decorator.h"
-#include "BehaviorTreeGraphNode.h"
-#include "EdGraphUtilities.h"
 #include "DetailCustomizations/BlackboardKeysDetails.h"
-#include "DetailCustomizations/BlackboardSelectorDetails.h"
-#include "BehaviorTreeEditor.h"
-#include "SGraphNode_BehaviorTree.h"
-#include "SGraphNode_Decorator.h"
 #include "BehaviorTree/Decorators/BTDecorator_BlueprintBase.h"
 #include "BehaviorTree/Services/BTService_BlueprintBase.h"
 #include "BehaviorTree/Tasks/BTTask_BlueprintBase.h"
+#include "BehaviorTreeDecoratorGraphNode_Decorator.h"
+#include "BehaviorTreeEditor.h"
+#include "BehaviorTreeGraphNode.h"
 #include "DetailCustomizations/BlackboardDecoratorDetails.h"
+#include "DetailCustomizations/BlackboardKeysDetails.h"
+#include "DetailCustomizations/BlackboardSelectorDetails.h"
+#include "EdGraphUtilities.h"
 #include "Modules/ModuleManager.h"
 #include "PropertyEditorModule.h"
+#include "SGraphNode_BehaviorTree.h"
+#include "SGraphNode_Decorator.h"
+#include "ValueOrBBKeyDetails.h"
 
-IMPLEMENT_MODULE( FBehaviorTreeEditorModule, BehaviorTreeEditor );
+IMPLEMENT_MODULE(FBehaviorTreeEditorModule, BehaviorTreeEditor);
 DEFINE_LOG_CATEGORY(LogBehaviorTreeEditor);
 
-const FName FBehaviorTreeEditorModule::BehaviorTreeEditorAppIdentifier( TEXT( "BehaviorTreeEditorApp" ) );
+const FName FBehaviorTreeEditorModule::BehaviorTreeEditorAppIdentifier(TEXT("BehaviorTreeEditorApp"));
 
 class FGraphPanelNodeFactory_BehaviorTree : public FGraphPanelNodeFactory
 {
@@ -47,21 +49,30 @@ void FBehaviorTreeEditorModule::StartupModule()
 	MenuExtensibilityManager = MakeShareable(new FExtensibilityManager);
 	ToolBarExtensibilityManager = MakeShareable(new FExtensibilityManager);
 
-	GraphPanelNodeFactory_BehaviorTree = MakeShareable( new FGraphPanelNodeFactory_BehaviorTree() );
+	GraphPanelNodeFactory_BehaviorTree = MakeShareable(new FGraphPanelNodeFactory_BehaviorTree());
 	FEdGraphUtilities::RegisterVisualNodeFactory(GraphPanelNodeFactory_BehaviorTree);
 
 	// Register the details customizer
 	FPropertyEditorModule& PropertyModule = FModuleManager::LoadModuleChecked<FPropertyEditorModule>("PropertyEditor");
-	PropertyModule.RegisterCustomPropertyTypeLayout( "BlackboardKeySelector", FOnGetPropertyTypeCustomizationInstance::CreateStatic( &FBlackboardSelectorDetails::MakeInstance ) );
-	PropertyModule.RegisterCustomClassLayout( "BTDecorator_Blackboard", FOnGetDetailCustomizationInstance::CreateStatic( &FBlackboardDecoratorDetails::MakeInstance ) );
-	PropertyModule.RegisterCustomClassLayout( "BTDecorator", FOnGetDetailCustomizationInstance::CreateStatic( &FBehaviorDecoratorDetails::MakeInstance ) );
+	PropertyModule.RegisterCustomPropertyTypeLayout("BlackboardKeySelector", FOnGetPropertyTypeCustomizationInstance::CreateStatic(&FBlackboardSelectorDetails::MakeInstance));
+	PropertyModule.RegisterCustomClassLayout("BTDecorator_Blackboard", FOnGetDetailCustomizationInstance::CreateStatic(&FBlackboardDecoratorDetails::MakeInstance));
+	PropertyModule.RegisterCustomClassLayout("BTDecorator", FOnGetDetailCustomizationInstance::CreateStatic(&FBehaviorDecoratorDetails::MakeInstance));
 	PropertyModule.RegisterCustomClassLayout("BlackboardKeyType_Class", FOnGetDetailCustomizationInstance::CreateStatic(&FBlackboardKeyDetails_Class::MakeInstance));
 	PropertyModule.RegisterCustomClassLayout("BlackboardKeyType_Enum", FOnGetDetailCustomizationInstance::CreateStatic(&FBlackboardKeyDetails_Enum::MakeInstance));
 	PropertyModule.RegisterCustomClassLayout("BlackboardKeyType_Object", FOnGetDetailCustomizationInstance::CreateStatic(&FBlackboardKeyDetails_Object::MakeInstance));
+	PropertyModule.RegisterCustomPropertyTypeLayout("ValueOrBBKey_Bool", FOnGetPropertyTypeCustomizationInstance::CreateStatic(&FValueOrBBKeyDetails::MakeInstance));
+	PropertyModule.RegisterCustomPropertyTypeLayout("ValueOrBBKey_Class", FOnGetPropertyTypeCustomizationInstance::CreateStatic(&FValueOrBBKeyDetails_Class::MakeInstance));
+	PropertyModule.RegisterCustomPropertyTypeLayout("ValueOrBBKey_Enum", FOnGetPropertyTypeCustomizationInstance::CreateStatic(&FValueOrBBKeyDetails_Enum::MakeInstance));
+	PropertyModule.RegisterCustomPropertyTypeLayout("ValueOrBBKey_Float", FOnGetPropertyTypeCustomizationInstance::CreateStatic(&FValueOrBBKeyDetails::MakeInstance));
+	PropertyModule.RegisterCustomPropertyTypeLayout("ValueOrBBKey_Int32", FOnGetPropertyTypeCustomizationInstance::CreateStatic(&FValueOrBBKeyDetails::MakeInstance));
+	PropertyModule.RegisterCustomPropertyTypeLayout("ValueOrBBKey_Name", FOnGetPropertyTypeCustomizationInstance::CreateStatic(&FValueOrBBKeyDetails::MakeInstance));
+	PropertyModule.RegisterCustomPropertyTypeLayout("ValueOrBBKey_String", FOnGetPropertyTypeCustomizationInstance::CreateStatic(&FValueOrBBKeyDetails::MakeInstance));
+	PropertyModule.RegisterCustomPropertyTypeLayout("ValueOrBBKey_Object", FOnGetPropertyTypeCustomizationInstance::CreateStatic(&FValueOrBBKeyDetails_Object::MakeInstance));
+	PropertyModule.RegisterCustomPropertyTypeLayout("ValueOrBBKey_Rotator", FOnGetPropertyTypeCustomizationInstance::CreateStatic(&FValueOrBBKeyDetails::MakeInstance));
+	PropertyModule.RegisterCustomPropertyTypeLayout("ValueOrBBKey_Vector", FOnGetPropertyTypeCustomizationInstance::CreateStatic(&FValueOrBBKeyDetails::MakeInstance));
 	PropertyModule.NotifyCustomizationModuleChanged();
 }
 
-
 void FBehaviorTreeEditorModule::ShutdownModule()
 {
 	if (!UObjectInitialized())
@@ -73,7 +84,7 @@ void FBehaviorTreeEditorModule::ShutdownModule()
 	ToolBarExtensibilityManager.Reset();
 	ClassCache.Reset();
 
-	if ( GraphPanelNodeFactory_BehaviorTree.IsValid() )
+	if (GraphPanelNodeFactory_BehaviorTree.IsValid())
 	{
 		FEdGraphUtilities::UnregisterVisualNodeFactory(GraphPanelNodeFactory_BehaviorTree);
 		GraphPanelNodeFactory_BehaviorTree.Reset();
@@ -83,17 +94,27 @@ void FBehaviorTreeEditorModule::ShutdownModule()
 	if (FModuleManager::Get().IsModuleLoaded("PropertyEditor"))
 	{
 		FPropertyEditorModule& PropertyModule = FModuleManager::LoadModuleChecked<FPropertyEditorModule>("PropertyEditor");
-		PropertyModule.UnregisterCustomPropertyTypeLayout( "BlackboardKeySelector" );
-		PropertyModule.UnregisterCustomClassLayout( "BTDecorator_Blackboard" );
-		PropertyModule.UnregisterCustomClassLayout( "BTDecorator" );
+		PropertyModule.UnregisterCustomPropertyTypeLayout("BlackboardKeySelector");
+		PropertyModule.UnregisterCustomClassLayout("BTDecorator_Blackboard");
+		PropertyModule.UnregisterCustomClassLayout("BTDecorator");
 		PropertyModule.UnregisterCustomClassLayout("BlackboardKeyType_Class");
 		PropertyModule.UnregisterCustomClassLayout("BlackboardKeyType_Enum");
 		PropertyModule.UnregisterCustomClassLayout("BlackboardKeyType_Object");
+		PropertyModule.UnregisterCustomPropertyTypeLayout("ValueOrBBKey_Bool");
+		PropertyModule.UnregisterCustomPropertyTypeLayout("ValueOrBBKey_Class");
+		PropertyModule.UnregisterCustomPropertyTypeLayout("ValueOrBBKey_Enum");
+		PropertyModule.UnregisterCustomPropertyTypeLayout("ValueOrBBKey_Float");
+		PropertyModule.UnregisterCustomPropertyTypeLayout("ValueOrBBKey_Int32");
+		PropertyModule.UnregisterCustomPropertyTypeLayout("ValueOrBBKey_Name");
+		PropertyModule.UnregisterCustomPropertyTypeLayout("ValueOrBBKey_String");
+		PropertyModule.UnregisterCustomPropertyTypeLayout("ValueOrBBKey_Object");
+		PropertyModule.UnregisterCustomPropertyTypeLayout("ValueOrBBKey_Rotator");
+		PropertyModule.UnregisterCustomPropertyTypeLayout("ValueOrBBKey_Vector");
 		PropertyModule.NotifyCustomizationModuleChanged();
 	}
 }
 
-TSharedRef<IBehaviorTreeEditor> FBehaviorTreeEditorModule::CreateBehaviorTreeEditor( const EToolkitMode::Type Mode, const TSharedPtr< IToolkitHost >& InitToolkitHost, UObject* Object )
+TSharedRef<IBehaviorTreeEditor> FBehaviorTreeEditorModule::CreateBehaviorTreeEditor(const EToolkitMode::Type Mode, const TSharedPtr<IToolkitHost>& InitToolkitHost, UObject* Object)
 {
 	if (!ClassCache.IsValid())
 	{
@@ -104,7 +125,7 @@ TSharedRef<IBehaviorTreeEditor> FBehaviorTreeEditorModule::CreateBehaviorTreeEdi
 		ClassCache->UpdateAvailableBlueprintClasses();
 	}
 
-	TSharedRef< FBehaviorTreeEditor > NewBehaviorTreeEditor( new FBehaviorTreeEditor() );
-	NewBehaviorTreeEditor->InitBehaviorTreeEditor( Mode, InitToolkitHost, Object );
-	return NewBehaviorTreeEditor;	
+	TSharedRef<FBehaviorTreeEditor> NewBehaviorTreeEditor(new FBehaviorTreeEditor());
+	NewBehaviorTreeEditor->InitBehaviorTreeEditor(Mode, InitToolkitHost, Object);
+	return NewBehaviorTreeEditor;
 }
diff --git a/Engine/Source/Editor/BehaviorTreeEditor/Private/DetailCustomizations/ValueOrBBKeyDetails.cpp b/Engine/Source/Editor/BehaviorTreeEditor/Private/DetailCustomizations/ValueOrBBKeyDetails.cpp
new file mode 100644
--- /dev/null
+++ b/Engine/Source/Editor/BehaviorTreeEditor/Private/DetailCustomizations/ValueOrBBKeyDetails.cpp
@@ -0,0 +1,492 @@
+// Copyright Epic Games, Inc. All Rights Reserved.
+
+#include "ValueOrBBKeyDetails.h"
+
+#include "BehaviorTree/Blackboard/BlackboardKeyType.h"
+#include "UObject/Field.h"
+
+#include "BehaviorTree/BTNode.h"
+#include "BehaviorTree/BehaviorTree.h"
+#include "BehaviorTree/BlackboardData.h"
+#include "BehaviorTree/ValueOrBBKey.h"
+#include "Delegates/Delegate.h"
+#include "DetailLayoutBuilder.h"
+#include "DetailWidgetRow.h"
+#include "Fonts/SlateFontInfo.h"
+#include "Framework/Commands/UIAction.h"
+#include "Framework/MultiBox/MultiBoxBuilder.h"
+#include "HAL/PlatformCrt.h"
+#include "IDetailChildrenBuilder.h"
+#include "IPropertyUtilities.h"
+#include "Internationalization/Internationalization.h"
+#include "Layout/Margin.h"
+#include "Misc/AssertionMacros.h"
+#include "Misc/Attribute.h"
+#include "PropertyCustomizationHelpers.h"
+#include "PropertyHandle.h"
+#include "SEnumCombo.h"
+#include "SlateOptMacros.h"
+#include "SlotBase.h"
+#include "Textures/SlateIcon.h"
+#include "Types/SlateEnums.h"
+#include "UObject/Class.h"
+#include "UObject/ObjectPtr.h"
+#include "Widgets/DeclarativeSyntaxSupport.h"
+#include "Widgets/Input/SComboButton.h"
+#include "Widgets/SBoxPanel.h"
+#include "Widgets/SWidget.h"
+#include "Widgets/Text/STextBlock.h"
+
+#define LOCTEXT_NAMESPACE "ValueOrBBKeyDetails"
+
+TSharedRef<IPropertyTypeCustomization> FValueOrBBKeyDetails::MakeInstance()
+{
+	return MakeShareable(new FValueOrBBKeyDetails);
+}
+
+TSharedRef<IPropertyTypeCustomization> FValueOrBBKeyDetails_Class::MakeInstance()
+{
+	return MakeShareable(new FValueOrBBKeyDetails_Class);
+}
+
+TSharedRef<IPropertyTypeCustomization> FValueOrBBKeyDetails_Enum::MakeInstance()
+{
+	return MakeShareable(new FValueOrBBKeyDetails_Enum);
+}
+
+TSharedRef<IPropertyTypeCustomization> FValueOrBBKeyDetails_Object::MakeInstance()
+{
+	return MakeShareable(new FValueOrBBKeyDetails_Object);
+}
+
+BEGIN_SLATE_FUNCTION_BUILD_OPTIMIZATION
+void FValueOrBBKeyDetails::CustomizeHeader(TSharedRef<IPropertyHandle> StructPropertyHandle, FDetailWidgetRow& HeaderRow, IPropertyTypeCustomizationUtils& StructCustomizationUtils)
+{
+	StructProperty = StructPropertyHandle;
+	KeyProperty = StructPropertyHandle->GetChildHandle(GET_MEMBER_NAME_CHECKED(FValueOrBlackboardKeyBase, Key));
+	DefaultValueProperty = StructPropertyHandle->GetChildHandle(TEXT("DefaultValue"));
+	CachedUtils = StructCustomizationUtils.GetPropertyUtilities();
+
+	if (const FProperty* MetaDataProperty = StructPropertyHandle->GetMetaDataProperty())
+	{
+		if (const TMap<FName, FString>* MetaDataMap = MetaDataProperty->GetMetaDataMap())
+		{
+			for(const TPair<FName, FString>& MetaData : *MetaDataMap)
+			{
+				DefaultValueProperty->SetInstanceMetaData(MetaData.Key, MetaData.Value);
+			}
+		}
+	}
+
+	ValidateData();
+
+	TSharedRef<SWidget> DefaultValueWidget = CreateDefaultValueWidget();
+	DefaultValueWidget->SetEnabled(TAttribute<bool>(this, &FValueOrBBKeyDetails::CanEditDefaultValue));
+
+	HeaderRow
+	.NameContent()
+	[
+		StructPropertyHandle->CreatePropertyNameWidget()
+	]
+	.ValueContent()
+	.HAlign(HAlign_Fill)
+	[
+		SNew(SHorizontalBox)
+		+ SHorizontalBox::Slot()
+		[
+			DefaultValueWidget
+		]
+		+ SHorizontalBox::Slot()
+		[
+			HasAccessToBlackboard() ?
+				SNew(SComboButton)
+				.OnGetMenuContent(this, &FValueOrBBKeyDetails::OnGetKeyNames)
+				.ButtonContent()
+				[
+					SNew(STextBlock)
+					.Text(this, &FValueOrBBKeyDetails::GetKeyDesc)
+					.Font(IDetailLayoutBuilder::GetDetailFont())
+				]
+				: KeyProperty->CreatePropertyValueWidget()
+		]
+		.Padding(FMargin(6.0f, 2.0f))
+	];
+}
+
+void FValueOrBBKeyDetails::CustomizeChildren(TSharedRef<class IPropertyHandle> StructPropertyHandle, class IDetailChildrenBuilder& StructBuilder, IPropertyTypeCustomizationUtils& StructCustomizationUtils)
+{
+}
+
+bool FValueOrBBKeyDetails::CanEditDefaultValue() const
+{
+	FName KeyValue;
+	KeyProperty->GetValue(KeyValue);
+	return KeyValue == FName() || !HasAccessToBlackboard();
+}
+
+TSharedRef<SWidget> FValueOrBBKeyDetails::CreateDefaultValueWidget()
+{
+	return DefaultValueProperty->CreatePropertyValueWidget();
+}
+
+void FValueOrBBKeyDetails_Class::CustomizeHeader(TSharedRef<class IPropertyHandle> StructPropertyHandle, class FDetailWidgetRow& HeaderRow, IPropertyTypeCustomizationUtils& StructCustomizationUtils)
+{
+	if ((BaseClassProperty = StructPropertyHandle->GetChildHandle(TEXT("BaseClass"))))
+	{
+		BaseClassProperty->SetOnPropertyValueChanged(FSimpleDelegate::CreateSP(this, &FValueOrBBKeyDetails_Class::OnBaseClassChanged));
+	}
+	FValueOrBBKeyDetails::CustomizeHeader(StructPropertyHandle, HeaderRow, StructCustomizationUtils);
+}
+
+void FValueOrBBKeyDetails_Class::CustomizeChildren(TSharedRef<class IPropertyHandle> StructPropertyHandle, class IDetailChildrenBuilder& StructBuilder, IPropertyTypeCustomizationUtils& StructCustomizationUtils)
+{
+	if (BaseClassProperty && BaseClassProperty->IsEditable())
+	{
+		StructBuilder.AddProperty(BaseClassProperty.ToSharedRef());
+	}
+}
+
+TSharedRef<SWidget> FValueOrBBKeyDetails_Class::CreateDefaultValueWidget()
+{
+	if (const FValueOrBlackboardKeyBase* DataPtr = GetDataPtr())
+	{
+		if (const UClass* BaseClass = static_cast<const FValueOrBBKey_Class*>(DataPtr)->BaseClass)
+		{
+			return SNew(SClassPropertyEntryBox)
+				.MetaClass(BaseClass)
+				.AllowNone(true)
+				.AllowAbstract(true)
+				.OnSetClass(this, &FValueOrBBKeyDetails_Class::OnSetClass)
+				.SelectedClass(this, &FValueOrBBKeyDetails_Class::OnGetSelectedClass);
+		}
+	}
+
+	return FValueOrBBKeyDetails::CreateDefaultValueWidget();
+}
+
+void FValueOrBBKeyDetails_Enum::CustomizeHeader(TSharedRef<class IPropertyHandle> StructPropertyHandle, class FDetailWidgetRow& HeaderRow, IPropertyTypeCustomizationUtils& StructCustomizationUtils)
+{
+	if ((EnumTypeProperty = StructPropertyHandle->GetChildHandle(TEXT("EnumType"))))
+	{
+		EnumTypeProperty->SetOnPropertyValueChanged(FSimpleDelegate::CreateSP(this, &FValueOrBBKeyDetails_Enum::OnEnumTypeChanged));
+	}
+
+	if ((NativeEnumTypeNameProperty = StructPropertyHandle->GetChildHandle(GET_MEMBER_NAME_CHECKED(FValueOrBBKey_Enum, NativeEnumTypeName))))
+	{
+		NativeEnumTypeNameProperty->SetOnPropertyValueChanged(FSimpleDelegate::CreateSP(this, &FValueOrBBKeyDetails_Enum::OnNativeEnumTypeNameChanged));
+	}
+
+	FValueOrBBKeyDetails::CustomizeHeader(StructPropertyHandle, HeaderRow, StructCustomizationUtils);
+}
+
+void FValueOrBBKeyDetails_Enum::CustomizeChildren(TSharedRef<class IPropertyHandle> StructPropertyHandle, class IDetailChildrenBuilder& StructBuilder, IPropertyTypeCustomizationUtils& StructCustomizationUtils)
+{
+	if (EnumTypeProperty && EnumTypeProperty->IsEditable())
+	{
+		StructBuilder.AddProperty(EnumTypeProperty.ToSharedRef());
+	}
+
+	if (NativeEnumTypeNameProperty && NativeEnumTypeNameProperty->IsEditable())
+	{
+		StructBuilder.AddProperty(NativeEnumTypeNameProperty.ToSharedRef());
+	}
+}
+
+TSharedRef<SWidget> FValueOrBBKeyDetails_Enum::CreateDefaultValueWidget()
+{
+	if(const FValueOrBlackboardKeyBase* DataPtr = GetDataPtr())
+	{
+		if (UEnum* EnumType = static_cast<const FValueOrBBKey_Enum*>(DataPtr)->EnumType)
+		{
+			return SNew(SEnumComboBox, EnumType)
+				.CurrentValue_Raw(this, &FValueOrBBKeyDetails_Enum::GetEnumValue)
+				.OnEnumSelectionChanged(SEnumComboBox::FOnEnumSelectionChanged::CreateSP(this, &FValueOrBBKeyDetails_Enum::OnEnumSelectionChanged));
+		}
+	}
+	return FValueOrBBKeyDetails::CreateDefaultValueWidget();
+}
+
+void FValueOrBBKeyDetails_Object::CustomizeHeader(TSharedRef<class IPropertyHandle> StructPropertyHandle, class FDetailWidgetRow& HeaderRow, IPropertyTypeCustomizationUtils& StructCustomizationUtils)
+{
+	if ((BaseClassProperty = StructPropertyHandle->GetChildHandle(TEXT("BaseClass"))))
+	{
+		BaseClassProperty->SetOnPropertyValueChanged(FSimpleDelegate::CreateSP(this, &FValueOrBBKeyDetails_Object::OnBaseClassChanged));
+	}
+
+	FValueOrBBKeyDetails::CustomizeHeader(StructPropertyHandle, HeaderRow, StructCustomizationUtils);
+}
+
+void FValueOrBBKeyDetails_Object::CustomizeChildren(TSharedRef<class IPropertyHandle> StructPropertyHandle, class IDetailChildrenBuilder& StructBuilder, IPropertyTypeCustomizationUtils& StructCustomizationUtils)
+{
+	if (BaseClassProperty && BaseClassProperty->IsEditable())
+	{
+		StructBuilder.AddProperty(BaseClassProperty.ToSharedRef());
+	}
+}
+
+TSharedRef<SWidget> FValueOrBBKeyDetails_Object::CreateDefaultValueWidget()
+{
+	if (const FValueOrBlackboardKeyBase* DataPtr = GetDataPtr())
+	{
+		if (UClass* BaseClass = static_cast<const FValueOrBBKey_Object*>(DataPtr)->BaseClass)
+		{
+			return SNew(SObjectPropertyEntryBox)
+				.AllowedClass(BaseClass)
+				.AllowClear(true)
+				.OnObjectChanged(this, &FValueOrBBKeyDetails_Object::OnObjectChanged)
+				.ObjectPath(this, &FValueOrBBKeyDetails_Object::OnGetObjectPath);
+		}
+	}
+	return FValueOrBBKeyDetails::CreateDefaultValueWidget();
+}
+END_SLATE_FUNCTION_BUILD_OPTIMIZATION
+
+bool FValueOrBBKeyDetails::HasAccessToBlackboard() const
+{
+	TArray<UObject*> OuterObjects;
+	KeyProperty->GetOuterObjects(OuterObjects);
+	if (OuterObjects.Num() > 0)
+	{
+		if (const UBTNode* OwnerNode = Cast<UBTNode>(OuterObjects[0]))
+		{
+			if (const UBehaviorTree* Tree = Cast<UBehaviorTree>(OwnerNode->GetTreeAsset()))
+			{
+				return Tree->BlackboardAsset != nullptr;
+			}
+		}
+	}
+	return false;
+}
+
+void FValueOrBBKeyDetails::GetMatchingKeys(TArray<FName>& OutNames)
+{
+	OutNames.Add(FName());
+	TArray<UObject*> OuterObjects;
+	KeyProperty->GetOuterObjects(OuterObjects);
+	if (OuterObjects.Num())
+	{
+		if(const FValueOrBlackboardKeyBase* DataPtr = GetDataPtr())
+		{
+			if (const UBTNode* OwnerNode = Cast<UBTNode>(OuterObjects[0]))
+			{
+				if (const UBehaviorTree* Tree = Cast<UBehaviorTree>(OwnerNode->GetTreeAsset()))
+				{
+					const UBlackboardData* Blackboard = Tree->GetBlackboardAsset();
+					while (Blackboard)
+					{
+						for (const FBlackboardEntry& BlackboardEntry : Blackboard->Keys)
+						{
+							if (DataPtr->IsCompatibleType(BlackboardEntry.KeyType))
+							{
+								OutNames.Add(BlackboardEntry.EntryName);
+							}
+						}
+						Blackboard = Blackboard->Parent;
+					}
+				}
+			}
+		}
+	}
+}
+
+void FValueOrBBKeyDetails::ValidateData()
+{
+	if (HasAccessToBlackboard())
+	{
+		TArray<FName> Keys;
+		GetMatchingKeys(Keys);
+		FName NameValue;
+		KeyProperty->GetValue(NameValue);
+		if (!Keys.Contains(NameValue))
+		{
+			KeyProperty->SetValue(FName());
+		}
+	}
+}
+
+void FValueOrBBKeyDetails_Class::ValidateData()
+{
+	FValueOrBBKeyDetails::ValidateData();
+	UObject* Object = nullptr;
+	DefaultValueProperty->GetValue(Object);
+	if (Object)
+	{
+		if (const FValueOrBlackboardKeyBase* DataPtr = GetDataPtr())
+		{
+			const UClass* BaseClass = static_cast<const FValueOrBBKey_Class*>(DataPtr)->BaseClass;
+			if (!BaseClass || !Object->IsA(BaseClass))
+			{
+				DefaultValueProperty->SetValue(static_cast<UObject*>(nullptr));
+			}
+		}
+	}
+}
+
+void FValueOrBBKeyDetails_Enum::ValidateData()
+{
+	FValueOrBBKeyDetails::ValidateData();
+	uint8 Value = 0;
+	if (DefaultValueProperty->GetValue(Value) == FPropertyAccess::Success)
+	{
+		if (const FValueOrBlackboardKeyBase* DataPtr = GetDataPtr())
+		{
+			const UEnum* EnumType = static_cast<const FValueOrBBKey_Enum*>(DataPtr)->EnumType;
+			if (EnumType)
+			{
+				if (!EnumType->IsValidEnumValue(Value) && EnumType->NumEnums())
+				{
+					DefaultValueProperty->SetValue(uint8(EnumType->GetValueByIndex(0)));
+				}
+			}
+			else
+			{
+				DefaultValueProperty->SetValue(uint8(0));
+			}
+		}
+		
+	}
+}
+void FValueOrBBKeyDetails_Object::ValidateData()
+{
+	FValueOrBBKeyDetails::ValidateData();
+	UObject* Object = nullptr;
+	DefaultValueProperty->GetValue(Object);
+	if (Object)
+	{
+		if(const FValueOrBlackboardKeyBase* DataPtr = GetDataPtr())
+		{
+			const UClass* BaseClass = static_cast<const FValueOrBBKey_Object*>(DataPtr)->BaseClass;
+			if (!BaseClass || !Object->IsA(BaseClass))
+			{
+				DefaultValueProperty->SetValue(static_cast<UObject*>(nullptr));
+			}
+		}
+	}
+}
+
+void FValueOrBBKeyDetails_Object::OnBaseClassChanged()
+{
+	ValidateData();
+	if (CachedUtils)
+	{
+		CachedUtils->ForceRefresh();
+	}
+}
+
+TSharedRef<SWidget> FValueOrBBKeyDetails::OnGetKeyNames()
+{
+	FMenuBuilder MenuBuilder(true, NULL);
+	MatchingKeys.Empty();
+	GetMatchingKeys(MatchingKeys);
+
+	for (int32 i = 0; i < MatchingKeys.Num(); i++)
+	{
+		FUIAction ItemAction(FExecuteAction::CreateSP(this, &FValueOrBBKeyDetails::OnKeyChanged, i));
+		MenuBuilder.AddMenuEntry(FText::FromName(MatchingKeys[i]), TAttribute<FText>(), FSlateIcon(), ItemAction);
+	}
+
+	return MenuBuilder.MakeWidget();
+}
+
+void FValueOrBBKeyDetails::OnKeyChanged(int32 Index)
+{
+	FName KeyValue = MatchingKeys[Index];
+	KeyProperty->SetValue(KeyValue);
+}
+
+FText FValueOrBBKeyDetails::GetKeyDesc() const
+{
+	FName KeyValue;
+	KeyProperty->GetValue(KeyValue);
+
+	return FText::FromString(KeyValue.ToString());
+}
+
+const FValueOrBlackboardKeyBase* FValueOrBBKeyDetails::GetDataPtr() const
+{
+	TArray<void*> StructPtrs;
+	StructProperty->AccessRawData(StructPtrs);
+	return (StructPtrs.Num() == 1) ? reinterpret_cast<FValueOrBlackboardKeyBase*>(StructPtrs[0]) : nullptr;
+}
+
+void FValueOrBBKeyDetails_Class::OnBaseClassChanged()
+{
+	ValidateData();
+	if (CachedUtils)
+	{
+		CachedUtils->ForceRefresh();
+	}
+}
+
+void FValueOrBBKeyDetails_Class::OnSetClass(const UClass* NewClass)
+{
+	DefaultValueProperty->SetValue(NewClass);
+}
+
+const UClass* FValueOrBBKeyDetails_Class::OnGetSelectedClass() const
+{
+	UObject* DefaultValue = nullptr;
+	DefaultValueProperty->GetValue(DefaultValue);
+	return Cast<UClass>(DefaultValue);
+}
+
+void FValueOrBBKeyDetails_Enum::OnEnumTypeChanged()
+{
+	ValidateData();
+	if (CachedUtils)
+	{
+		CachedUtils->ForceRefresh();
+	}
+}
+
+void FValueOrBBKeyDetails_Enum::OnEnumSelectionChanged(int32 NewValue, ESelectInfo::Type)
+{
+	DefaultValueProperty->SetValue(static_cast<uint8>(NewValue));
+}
+
+void FValueOrBBKeyDetails_Enum::OnNativeEnumTypeNameChanged()
+{
+	EnumTypeProperty->SetValue(static_cast<UEnum*>(nullptr));
+
+	FString NativeEnumTypeName;
+	if (NativeEnumTypeNameProperty->GetValue(NativeEnumTypeName) == FPropertyAccess::Success)
+	{
+		if (const UEnum* const NativeEnumType = UClass::TryFindTypeSlow<UEnum>(NativeEnumTypeName, EFindFirstObjectOptions::ExactClass))
+		{
+			EnumTypeProperty->SetValue(NativeEnumType);
+		}
+	}
+
+	ValidateData();
+	if (CachedUtils)
+	{
+		CachedUtils->ForceRefresh();
+	}
+}
+
+int32 FValueOrBBKeyDetails_Enum::GetEnumValue() const
+{
+	uint8 EnumValue = 0;
+	DefaultValueProperty->GetValue(EnumValue);
+	return EnumValue;
+}
+
+bool FValueOrBBKeyDetails_Enum::CanEditEnumType() const
+{
+	FString NativeEnumTypeName;
+	const FPropertyAccess::Result NativeEnumTypeNamePropertyAccessResult = NativeEnumTypeNameProperty->GetValue(NativeEnumTypeName);
+	return (NativeEnumTypeNamePropertyAccessResult != FPropertyAccess::Success) || NativeEnumTypeName.IsEmpty();
+}
+
+void FValueOrBBKeyDetails_Object::OnObjectChanged(const FAssetData& AssetData)
+{
+	DefaultValueProperty->SetValue(AssetData.GetAsset());
+}
+
+FString FValueOrBBKeyDetails_Object::OnGetObjectPath() const
+{
+	UObject* DefaultValue = nullptr;
+	DefaultValueProperty->GetValue(DefaultValue);
+	return DefaultValue ? DefaultValue->GetPathName() : FString();
+}
+#undef LOCTEXT_NAMESPACE
diff --git a/Engine/Source/Runtime/AIModule/Classes/BehaviorTree/ValueOrBBKey.h b/Engine/Source/Runtime/AIModule/Classes/BehaviorTree/ValueOrBBKey.h
new file mode 100644
--- /dev/null
+++ b/Engine/Source/Runtime/AIModule/Classes/BehaviorTree/ValueOrBBKey.h
@@ -0,0 +1,400 @@
+// Copyright Epic Games, Inc. All Rights Reserved.
+
+#pragma once
+
+#include "BehaviorTree/BehaviorTreeComponent.h"
+#include "BehaviorTree/Blackboard/BlackboardKey.h"
+#include "BehaviorTree/BlackboardComponent.h"
+#include "GameplayTagContainer.h"
+#include "Misc/Guid.h"
+
+#include "ValueOrBBKey.generated.h"
+
+class UBlackboardKeyType;
+class FValueOrBBKeyDetails;
+
+struct FValueOrBBKey_Bool;
+struct FValueOrBBKey_Class;
+struct FValueOrBBKey_Enum;
+struct FValueOrBBKey_Float;
+struct FValueOrBBKey_Int32;
+struct FValueOrBBKey_Name;
+struct FValueOrBBKey_String;
+struct FValueOrBBKey_Object;
+struct FValueOrBBKey_Rotator;
+struct FValueOrBBKey_Vector;
+
+namespace EValueOrBBKeyVersion
+{
+	enum Type
+	{
+		Initial,
+		// Changed raw property to their EValueOrBBKeyVersion equivalent
+		ChangedPropertyToValueOrBBKey,
+		// -----<new versions can be added before this line>-------------------------------------------------
+		// - this needs to be the last line (see note below)
+		VersionPlusOne,
+		LatestVersion = VersionPlusOne - 1
+	};
+
+	// The GUID for this custom version number
+	AIMODULE_API extern const FGuid Guid;
+} // namespace EValueOrBBKeyVersion
+
+namespace FBlackboard
+{
+	template <typename T>
+	TOptional<typename T::FDataType> TryGetBlackboardKeyValue(const UBlackboardComponent& Blackboard, const FName& Name, FBlackboard::FKey& InOutCachedKey)
+	{
+		if (InOutCachedKey == FBlackboard::InvalidKey)
+		{
+			InOutCachedKey = Blackboard.GetKeyID(Name);
+		}
+		if (Blackboard.GetKeyType(InOutCachedKey) == T::StaticClass())
+		{
+			return { Blackboard.GetValue<T>(InOutCachedKey) };
+		}
+		return {};
+	}
+
+	template <typename T>
+	typename T::FDataType GetValue(const UBlackboardComponent& Blackboard, const FName& Name, FBlackboard::FKey& InOutCachedKey, const typename T::FDataType& DefaultValue)
+	{
+		if (!Name.IsNone())
+		{
+			if (const TOptional<typename T::FDataType> KeyValue = FBlackboard::TryGetBlackboardKeyValue<T>(Blackboard, Name, InOutCachedKey))
+			{
+				return KeyValue.GetValue();
+			}
+		}
+		return DefaultValue;
+	}
+
+	template <typename T>
+	typename T::FDataType GetValue(const UBehaviorTreeComponent& BehaviorComp, const FName& Name, FBlackboard::FKey& InOutCachedKey, const typename T::FDataType& DefaultValue)
+	{
+		if (const UBlackboardComponent* Blackboard = BehaviorComp.GetBlackboardComponent())
+		{
+			return GetValue<T>(*Blackboard, Name, InOutCachedKey, DefaultValue);
+		}
+		return DefaultValue;
+	}
+} // namespace FBlackboard
+
+// Base struct to simplify edition in the editor, shouldn't be used elsewhere
+USTRUCT()
+struct FValueOrBlackboardKeyBase
+{
+	GENERATED_BODY()
+
+	friend FValueOrBBKeyDetails;
+
+#if WITH_EDITOR
+	AIMODULE_API virtual ~FValueOrBlackboardKeyBase() = default;
+	AIMODULE_API virtual bool IsCompatibleType(const UBlackboardKeyType* KeyType) const { return false; };
+	AIMODULE_API virtual FString ToString() const { return FString(); }
+#endif // WITH_EDITOR
+
+	const FName& GetKey() const { return Key; }
+	void SetKey(FName NewKey) { Key = NewKey; }
+
+	FBlackboard::FKey GetKeyId(const UBehaviorTreeComponent& OwnerComp) const;
+
+protected:
+	template <typename T>
+	FString ToStringInternal(const T& DefaultValue) const
+	{
+		if (!Key.IsNone())
+		{
+			return ToStringKeyName();
+		}
+		else
+		{
+			return FString::Format(TEXT("{0}"), { DefaultValue });
+		}
+	}
+
+	AIMODULE_API FString ToStringKeyName() const;
+
+	UPROPERTY(EditAnywhere, Category = "Value")
+	FName Key;
+
+	mutable FBlackboard::FKey KeyId = FBlackboard::InvalidKey;
+};
+
+USTRUCT(BlueprintType)
+struct FValueOrBBKey_Bool : public FValueOrBlackboardKeyBase
+{
+	GENERATED_BODY()
+
+	FValueOrBBKey_Bool(bool Default = false)
+		: DefaultValue(Default) {}
+	AIMODULE_API bool GetValue(const UBehaviorTreeComponent& BehaviorComp) const;
+	AIMODULE_API bool GetValue(const UBehaviorTreeComponent* BehaviorComp) const;
+	AIMODULE_API bool GetValue(const UBlackboardComponent& Blackboard) const;
+	AIMODULE_API bool GetValue(const UBlackboardComponent* Blackboard) const;
+
+#if WITH_EDITOR
+	AIMODULE_API virtual bool IsCompatibleType(const UBlackboardKeyType* KeyType) const override;
+#endif // WITH_EDITOR
+
+	AIMODULE_API FString ToString() const { return ToStringInternal(DefaultValue); }
+
+protected:
+	UPROPERTY(EditAnywhere, Category = "Value")
+	bool DefaultValue = false;
+};
+
+USTRUCT(BlueprintType)
+struct FValueOrBBKey_Class : public FValueOrBlackboardKeyBase
+{
+	GENERATED_BODY()
+
+	friend class FValueOrBBKeyDetails_Class;
+
+	AIMODULE_API FValueOrBBKey_Class() = default;
+
+	template <typename T>
+	FValueOrBBKey_Class(TSubclassOf<T> Default)
+		: DefaultValue(Default), BaseClass(T::StaticClass()) {}
+
+	template <typename T>
+	TSubclassOf<T> GetValue(const UBehaviorTreeComponent& BehaviorComp) const
+	{
+		return TSubclassOf<T>(GetValue(BehaviorComp));
+	}
+
+	AIMODULE_API UClass* GetValue(const UBehaviorTreeComponent& BehaviorComp) const;
+	AIMODULE_API UClass* GetValue(const UBehaviorTreeComponent* BehaviorComp) const;
+	AIMODULE_API UClass* GetValue(const UBlackboardComponent& Blackboard) const;
+	AIMODULE_API UClass* GetValue(const UBlackboardComponent* Blackboard) const;
+
+#if WITH_EDITOR
+	AIMODULE_API virtual bool IsCompatibleType(const UBlackboardKeyType* KeyType) const override;
+#endif // WITH_EDITOR
+
+	AIMODULE_API FString ToString() const;
+
+protected:
+	UPROPERTY(EditAnywhere, Category = "Value")
+	TObjectPtr<UClass> DefaultValue = nullptr;
+
+	UPROPERTY(EditDefaultsOnly, Category = "Value", meta = (AllowAbstract = "1"))
+	TObjectPtr<UClass> BaseClass = UObject::StaticClass();
+};
+
+USTRUCT(BlueprintType)
+struct FValueOrBBKey_Enum : public FValueOrBlackboardKeyBase
+{
+	GENERATED_BODY()
+
+	friend class FValueOrBBKeyDetails_Enum;
+
+	AIMODULE_API FValueOrBBKey_Enum() = default;
+
+	template <typename T>
+	FValueOrBBKey_Enum(T Default)
+		: DefaultValue(static_cast<uint8>(Default)), EnumType(StaticEnum<T>()), NativeEnumTypeName(GetNameSafe(EnumType)) {}
+
+	template <typename T>
+	T GetValue(const UBehaviorTreeComponent& BehaviorComp) const
+	{
+		return static_cast<T>(GetValue(BehaviorComp));
+	}
+
+	AIMODULE_API uint8 GetValue(const UBehaviorTreeComponent& BehaviorComp) const;
+	AIMODULE_API uint8 GetValue(const UBehaviorTreeComponent* BehaviorComp) const;
+	AIMODULE_API uint8 GetValue(const UBlackboardComponent& Blackboard) const;
+	AIMODULE_API uint8 GetValue(const UBlackboardComponent* Blackboard) const;
+
+#if WITH_EDITOR
+	AIMODULE_API virtual bool IsCompatibleType(const UBlackboardKeyType* KeyType) const override;
+#endif // WITH_EDITOR
+
+	AIMODULE_API FString ToString() const;
+
+protected:
+	UPROPERTY(EditAnywhere, Category = "Value")
+	uint8 DefaultValue = 0;
+
+	UPROPERTY(EditDefaultsOnly, Category = "Value")
+	TObjectPtr<UEnum> EnumType = nullptr;
+
+	/** Name of enum type defined in C++ code, will take priority over asset from EnumType property */
+	UPROPERTY(EditDefaultsOnly, Category = "Value")
+	FString NativeEnumTypeName = TEXT("");
+};
+
+USTRUCT(BlueprintType)
+struct FValueOrBBKey_Float : public FValueOrBlackboardKeyBase
+{
+	GENERATED_BODY()
+
+	FValueOrBBKey_Float(float Default = 0.f)
+		: DefaultValue(Default) {}
+	AIMODULE_API float GetValue(const UBehaviorTreeComponent& BehaviorComp) const;
+	AIMODULE_API float GetValue(const UBehaviorTreeComponent* BehaviorComp) const;
+	AIMODULE_API float GetValue(const UBlackboardComponent& Blackboard) const;
+	AIMODULE_API float GetValue(const UBlackboardComponent* Blackboard) const;
+
+#if WITH_EDITOR
+	AIMODULE_API virtual bool IsCompatibleType(const UBlackboardKeyType* KeyType) const override;
+#endif // WITH_EDITOR
+
+	AIMODULE_API FString ToString() const;
+
+protected:
+	UPROPERTY(EditAnywhere, Category = "Value")
+	float DefaultValue = 0.f;
+};
+
+USTRUCT(BlueprintType)
+struct FValueOrBBKey_Int32 : public FValueOrBlackboardKeyBase
+{
+	GENERATED_BODY()
+	FValueOrBBKey_Int32(int32 Default = 0)
+		: DefaultValue(Default) {}
+	AIMODULE_API int32 GetValue(const UBehaviorTreeComponent& BehaviorComp) const;
+	AIMODULE_API int32 GetValue(const UBehaviorTreeComponent* BehaviorComp) const;
+	AIMODULE_API int32 GetValue(const UBlackboardComponent& Blackboard) const;
+	AIMODULE_API int32 GetValue(const UBlackboardComponent* Blackboard) const;
+
+#if WITH_EDITOR
+	AIMODULE_API virtual bool IsCompatibleType(const UBlackboardKeyType* KeyType) const override;
+#endif // WITH_EDITOR
+
+	AIMODULE_API FString ToString() const { return ToStringInternal(DefaultValue); }
+
+protected:
+	UPROPERTY(EditAnywhere, Category = "Value")
+	int32 DefaultValue = 0;
+};
+
+USTRUCT(BlueprintType)
+struct FValueOrBBKey_Name : public FValueOrBlackboardKeyBase
+{
+	GENERATED_BODY()
+
+	FValueOrBBKey_Name(FName Default = FName())
+		: DefaultValue(Default) {}
+	AIMODULE_API FName GetValue(const UBehaviorTreeComponent& BehaviorComp) const;
+	AIMODULE_API FName GetValue(const UBehaviorTreeComponent* BehaviorComp) const;
+	AIMODULE_API FName GetValue(const UBlackboardComponent& Blackboard) const;
+	AIMODULE_API FName GetValue(const UBlackboardComponent* Blackboard) const;
+
+#if WITH_EDITOR
+	AIMODULE_API virtual bool IsCompatibleType(const UBlackboardKeyType* KeyType) const override;
+#endif // WITH_EDITOR
+
+	AIMODULE_API FString ToString() const;
+
+protected:
+	UPROPERTY(EditAnywhere, Category = "Value")
+	FName DefaultValue;
+};
+
+USTRUCT(BlueprintType)
+struct FValueOrBBKey_String : public FValueOrBlackboardKeyBase
+{
+	GENERATED_BODY()
+
+	FValueOrBBKey_String(FString Default = FString())
+		: DefaultValue(MoveTemp(Default)) {}
+	AIMODULE_API FString GetValue(const UBehaviorTreeComponent& BehaviorComp) const;
+	AIMODULE_API FString GetValue(const UBehaviorTreeComponent* BehaviorComp) const;
+	AIMODULE_API FString GetValue(const UBlackboardComponent& Blackboard) const;
+	AIMODULE_API FString GetValue(const UBlackboardComponent* Blackboard) const;
+
+#if WITH_EDITOR
+	AIMODULE_API virtual bool IsCompatibleType(const UBlackboardKeyType* KeyType) const override;
+#endif // WITH_EDITOR
+
+	AIMODULE_API FString ToString() const { return ToStringInternal(DefaultValue); }
+
+protected:
+	UPROPERTY(EditAnywhere, Category = "Value")
+	FString DefaultValue;
+};
+
+USTRUCT(BlueprintType)
+struct FValueOrBBKey_Object : public FValueOrBlackboardKeyBase
+{
+	GENERATED_BODY()
+	friend class FValueOrBBKeyDetails_Object;
+
+	FValueOrBBKey_Object() = default;
+
+	template <typename T>
+	FValueOrBBKey_Object(TObjectPtr<T> Default)
+		: DefaultValue(Default), BaseClass(T::StaticClass()) {}
+
+	template <typename T>
+	T* GetValue(const UBehaviorTreeComponent& BehaviorComp) const
+	{
+		return Cast<T>(GetValue(BehaviorComp));
+	}
+
+	AIMODULE_API UObject* GetValue(const UBehaviorTreeComponent& BehaviorComp) const;
+	AIMODULE_API UObject* GetValue(const UBehaviorTreeComponent* BehaviorComp) const;
+	AIMODULE_API UObject* GetValue(const UBlackboardComponent& Blackboard) const;
+	AIMODULE_API UObject* GetValue(const UBlackboardComponent* Blackboard) const;
+
+#if WITH_EDITOR
+	AIMODULE_API virtual bool IsCompatibleType(const UBlackboardKeyType* KeyType) const override;
+#endif // WITH_EDITOR
+
+	AIMODULE_API FString ToString() const;
+
+protected:
+	UPROPERTY(EditAnywhere, Category = "Value")
+	TObjectPtr<UObject> DefaultValue = nullptr;
+
+	UPROPERTY(EditDefaultsOnly, Category = "Value", meta = (AllowAbstract = "1"))
+	TObjectPtr<UClass> BaseClass = UObject::StaticClass();
+};
+
+USTRUCT(BlueprintType)
+struct FValueOrBBKey_Rotator : public FValueOrBlackboardKeyBase
+{
+	GENERATED_BODY()
+
+	FValueOrBBKey_Rotator(FRotator Default = FRotator::ZeroRotator)
+		: DefaultValue(MoveTemp(Default)) {}
+	AIMODULE_API FRotator GetValue(const UBehaviorTreeComponent& BehaviorComp) const;
+	AIMODULE_API FRotator GetValue(const UBehaviorTreeComponent* BehaviorComp) const;
+	AIMODULE_API FRotator GetValue(const UBlackboardComponent& Blackboard) const;
+	AIMODULE_API FRotator GetValue(const UBlackboardComponent* Blackboard) const;
+
+#if WITH_EDITOR
+	AIMODULE_API virtual bool IsCompatibleType(const UBlackboardKeyType* KeyType) const override;
+#endif // WITH_EDITOR
+
+	AIMODULE_API FString ToString() const;
+
+protected:
+	UPROPERTY(EditAnywhere, Category = "Value")
+	FRotator DefaultValue;
+};
+
+USTRUCT(BlueprintType)
+struct FValueOrBBKey_Vector : public FValueOrBlackboardKeyBase
+{
+	GENERATED_BODY()
+
+	FValueOrBBKey_Vector(FVector Default = FVector::ZeroVector)
+		: DefaultValue(MoveTemp(Default)) {}
+	AIMODULE_API FVector GetValue(const UBehaviorTreeComponent& BehaviorComp) const;
+	AIMODULE_API FVector GetValue(const UBehaviorTreeComponent* BehaviorComp) const;
+	AIMODULE_API FVector GetValue(const UBlackboardComponent& Blackboard) const;
+	AIMODULE_API FVector GetValue(const UBlackboardComponent* Blackboard) const;
+
+#if WITH_EDITOR
+	AIMODULE_API virtual bool IsCompatibleType(const UBlackboardKeyType* KeyType) const override;
+#endif // WITH_EDITOR
+
+	AIMODULE_API FString ToString() const;
+
+protected:
+	UPROPERTY(EditAnywhere, Category = "Value")
+	FVector DefaultValue;
+};
\ No newline at end of file
diff --git a/Engine/Source/Runtime/AIModule/Private/BehaviorTree/ValueOrBBKey.cpp b/Engine/Source/Runtime/AIModule/Private/BehaviorTree/ValueOrBBKey.cpp
new file mode 100644
--- /dev/null
+++ b/Engine/Source/Runtime/AIModule/Private/BehaviorTree/ValueOrBBKey.cpp
@@ -0,0 +1,393 @@
+// Copyright Epic Games, Inc. All Rights Reserved.
+
+#include "BehaviorTree/ValueOrBBKey.h"
+
+#include "BehaviorTree/Blackboard/BlackboardKeyType_Bool.h"
+#include "BehaviorTree/Blackboard/BlackboardKeyType_Class.h"
+#include "BehaviorTree/Blackboard/BlackboardKeyType_Enum.h"
+#include "BehaviorTree/Blackboard/BlackboardKeyType_Float.h"
+#include "BehaviorTree/Blackboard/BlackboardKeyType_Int.h"
+#include "BehaviorTree/Blackboard/BlackboardKeyType_Name.h"
+#include "BehaviorTree/Blackboard/BlackboardKeyType_Object.h"
+#include "BehaviorTree/Blackboard/BlackboardKeyType_Rotator.h"
+#include "BehaviorTree/Blackboard/BlackboardKeyType_String.h"
+#include "BehaviorTree/Blackboard/BlackboardKeyType_Vector.h"
+
+#include "Serialization/CustomVersion.h"
+
+const FGuid EValueOrBBKeyVersion::Guid = FGuid(0x82238A1B, 0xD8DA4083, 0xB5BC56A9, 0x48AA4C22);
+FCustomVersionRegistration GValueOrBBKeyVersion(EValueOrBBKeyVersion::Guid, EValueOrBBKeyVersion::LatestVersion, TEXT("EValueOrBBKeyVersion"));
+
+FBlackboard::FKey FValueOrBlackboardKeyBase::GetKeyId(const UBehaviorTreeComponent& OwnerComp) const
+{
+	if (KeyId == FBlackboard::InvalidKey)
+	{
+		if (const UBlackboardComponent* BlackboardComp = OwnerComp.GetBlackboardComponent())
+		{
+			KeyId = BlackboardComp->GetKeyID(Key);
+		}
+	}
+	return KeyId;
+}
+
+FString FValueOrBlackboardKeyBase::ToStringKeyName() const
+{
+	return FString::Format(TEXT("Key: {0}"), { *Key.ToString() });
+}
+
+FString FValueOrBBKey_Class::ToString() const
+{
+	if (!Key.IsNone())
+	{
+		return ToStringKeyName();
+	}
+	else
+	{
+		return GetNameSafe(DefaultValue);
+	}
+}
+
+FString FValueOrBBKey_Enum::ToString() const
+{
+	if (!Key.IsNone())
+	{
+		return ToStringKeyName();
+	}
+	else
+	{
+		return EnumType ? EnumType->GetNameStringByValue(DefaultValue) : TEXT("Invalid Enum");
+	}
+}
+
+FString FValueOrBBKey_Float::ToString() const
+{
+	if (!Key.IsNone())
+	{
+		return ToStringKeyName();
+	}
+	else
+	{
+		return FString::Printf(TEXT("%.2f"), DefaultValue);
+	}
+}
+
+FString FValueOrBBKey_Name::ToString() const
+{
+	if (!Key.IsNone())
+	{
+		return ToStringKeyName();
+	}
+	else
+	{
+		return DefaultValue.ToString();
+	}
+}
+
+FString FValueOrBBKey_Object::ToString() const
+{
+	if (!Key.IsNone())
+	{
+		return ToStringKeyName();
+	}
+	else
+	{
+		return GetNameSafe(DefaultValue);
+	}
+}
+
+FString FValueOrBBKey_Rotator::ToString() const
+{
+	if (!Key.IsNone())
+	{
+		return ToStringKeyName();
+	}
+	else
+	{
+		return DefaultValue.ToString();
+	}
+}
+
+FString FValueOrBBKey_Vector::ToString() const
+{
+	if (!Key.IsNone())
+	{
+		return ToStringKeyName();
+	}
+	else
+	{
+		return DefaultValue.ToString();
+	}
+}
+
+bool FValueOrBBKey_Bool::GetValue(const UBlackboardComponent& Blackboard) const
+{
+	return FBlackboard::GetValue<UBlackboardKeyType_Bool>(Blackboard, Key, KeyId, DefaultValue);
+}
+
+bool FValueOrBBKey_Bool::GetValue(const UBlackboardComponent* Blackboard) const
+{
+	return Blackboard ? GetValue(*Blackboard) : DefaultValue;
+}
+
+bool FValueOrBBKey_Bool::GetValue(const UBehaviorTreeComponent& BehaviorComp) const
+{
+	return FBlackboard::GetValue<UBlackboardKeyType_Bool>(BehaviorComp, Key, KeyId, DefaultValue);
+}
+
+bool FValueOrBBKey_Bool::GetValue(const UBehaviorTreeComponent* BehaviorComp) const
+{
+	return BehaviorComp ? GetValue(*BehaviorComp) : DefaultValue;
+}
+
+UClass* FValueOrBBKey_Class::GetValue(const UBlackboardComponent& Blackboard) const
+{
+	return FBlackboard::GetValue<UBlackboardKeyType_Class>(Blackboard, Key, KeyId, DefaultValue.Get());
+}
+
+UClass* FValueOrBBKey_Class::GetValue(const UBlackboardComponent* Blackboard) const
+{
+	return Blackboard ? GetValue(*Blackboard) : DefaultValue.Get();
+}
+
+UClass* FValueOrBBKey_Class::GetValue(const UBehaviorTreeComponent& BehaviorComp) const
+{
+	return FBlackboard::GetValue<UBlackboardKeyType_Class>(BehaviorComp, Key, KeyId, DefaultValue.Get());
+}
+
+UClass* FValueOrBBKey_Class::GetValue(const UBehaviorTreeComponent* BehaviorComp) const
+{
+	return BehaviorComp ? GetValue(*BehaviorComp) : DefaultValue.Get();
+}
+
+uint8 FValueOrBBKey_Enum::GetValue(const UBlackboardComponent& Blackboard) const
+{
+	return FBlackboard::GetValue<UBlackboardKeyType_Enum>(Blackboard, Key, KeyId, DefaultValue);
+}
+
+uint8 FValueOrBBKey_Enum::GetValue(const UBlackboardComponent* Blackboard) const
+{
+	return Blackboard ? GetValue(*Blackboard) : DefaultValue;
+}
+
+uint8 FValueOrBBKey_Enum::GetValue(const UBehaviorTreeComponent& BehaviorComp) const
+{
+	return FBlackboard::GetValue<UBlackboardKeyType_Enum>(BehaviorComp, Key, KeyId, DefaultValue);
+}
+
+uint8 FValueOrBBKey_Enum::GetValue(const UBehaviorTreeComponent* BehaviorComp) const
+{
+	return BehaviorComp ? GetValue(*BehaviorComp) : DefaultValue;
+}
+
+float FValueOrBBKey_Float::GetValue(const UBlackboardComponent& Blackboard) const
+{
+	return FBlackboard::GetValue<UBlackboardKeyType_Float>(Blackboard, Key, KeyId, DefaultValue);
+}
+
+float FValueOrBBKey_Float::GetValue(const UBlackboardComponent* Blackboard) const
+{
+	return Blackboard ? GetValue(*Blackboard) : DefaultValue;
+}
+
+float FValueOrBBKey_Float::GetValue(const UBehaviorTreeComponent& BehaviorComp) const
+{
+	return FBlackboard::GetValue<UBlackboardKeyType_Float>(BehaviorComp, Key, KeyId, DefaultValue);
+}
+
+float FValueOrBBKey_Float::GetValue(const UBehaviorTreeComponent* BehaviorComp) const
+{
+	return BehaviorComp ? GetValue(*BehaviorComp) : DefaultValue;
+}
+
+int32 FValueOrBBKey_Int32::GetValue(const UBlackboardComponent& Blackboard) const
+{
+	return FBlackboard::GetValue<UBlackboardKeyType_Int>(Blackboard, Key, KeyId, DefaultValue);
+}
+
+int32 FValueOrBBKey_Int32::GetValue(const UBlackboardComponent* Blackboard) const
+{
+	return Blackboard ? GetValue(*Blackboard) : DefaultValue;
+}
+
+int32 FValueOrBBKey_Int32::GetValue(const UBehaviorTreeComponent& BehaviorComp) const
+{
+	return FBlackboard::GetValue<UBlackboardKeyType_Int>(BehaviorComp, Key, KeyId, DefaultValue);
+}
+
+int32 FValueOrBBKey_Int32::GetValue(const UBehaviorTreeComponent* BehaviorComp) const
+{
+	return BehaviorComp ? GetValue(*BehaviorComp) : DefaultValue;
+}
+
+FName FValueOrBBKey_Name::GetValue(const UBlackboardComponent& Blackboard) const
+{
+	return FBlackboard::GetValue<UBlackboardKeyType_Name>(Blackboard, Key, KeyId, DefaultValue);
+}
+
+FName FValueOrBBKey_Name::GetValue(const UBlackboardComponent* Blackboard) const
+{
+	return Blackboard ? GetValue(*Blackboard) : DefaultValue;
+}
+
+FName FValueOrBBKey_Name::GetValue(const UBehaviorTreeComponent& BehaviorComp) const
+{
+	return FBlackboard::GetValue<UBlackboardKeyType_Name>(BehaviorComp, Key, KeyId, DefaultValue);
+}
+
+FName FValueOrBBKey_Name::GetValue(const UBehaviorTreeComponent* BehaviorComp) const
+{
+	return BehaviorComp ? GetValue(*BehaviorComp) : DefaultValue;
+}
+
+FString FValueOrBBKey_String::GetValue(const UBlackboardComponent& Blackboard) const
+{
+	return FBlackboard::GetValue<UBlackboardKeyType_String>(Blackboard, Key, KeyId, DefaultValue);
+}
+
+FString FValueOrBBKey_String::GetValue(const UBlackboardComponent* Blackboard) const
+{
+	return Blackboard ? GetValue(*Blackboard) : DefaultValue;
+}
+
+FString FValueOrBBKey_String::GetValue(const UBehaviorTreeComponent& BehaviorComp) const
+{
+	return FBlackboard::GetValue<UBlackboardKeyType_String>(BehaviorComp, Key, KeyId, DefaultValue);
+}
+
+FString FValueOrBBKey_String::GetValue(const UBehaviorTreeComponent* BehaviorComp) const
+{
+	return BehaviorComp ? GetValue(*BehaviorComp) : DefaultValue;
+}
+
+UObject* FValueOrBBKey_Object::GetValue(const UBlackboardComponent& Blackboard) const
+{
+	return FBlackboard::GetValue<UBlackboardKeyType_Object>(Blackboard, Key, KeyId, DefaultValue.Get());
+}
+
+UObject* FValueOrBBKey_Object::GetValue(const UBlackboardComponent* Blackboard) const
+{
+	return Blackboard ? GetValue(*Blackboard) : DefaultValue.Get();
+}
+
+UObject* FValueOrBBKey_Object::GetValue(const UBehaviorTreeComponent& BehaviorComp) const
+{
+	return FBlackboard::GetValue<UBlackboardKeyType_Object>(BehaviorComp, Key, KeyId, DefaultValue.Get());
+}
+
+UObject* FValueOrBBKey_Object::GetValue(const UBehaviorTreeComponent* BehaviorComp) const
+{
+	return BehaviorComp ? GetValue(*BehaviorComp) : DefaultValue.Get();
+}
+
+FRotator FValueOrBBKey_Rotator::GetValue(const UBlackboardComponent& Blackboard) const
+{
+	return FBlackboard::GetValue<UBlackboardKeyType_Rotator>(Blackboard, Key, KeyId, DefaultValue);
+}
+
+FRotator FValueOrBBKey_Rotator::GetValue(const UBlackboardComponent* Blackboard) const
+{
+	return Blackboard ? GetValue(*Blackboard) : DefaultValue;
+}
+
+FRotator FValueOrBBKey_Rotator::GetValue(const UBehaviorTreeComponent& BehaviorComp) const
+{
+	return FBlackboard::GetValue<UBlackboardKeyType_Rotator>(BehaviorComp, Key, KeyId, DefaultValue);
+}
+
+FRotator FValueOrBBKey_Rotator::GetValue(const UBehaviorTreeComponent* BehaviorComp) const
+{
+	return BehaviorComp ? GetValue(*BehaviorComp) : DefaultValue;
+}
+
+FVector FValueOrBBKey_Vector::GetValue(const UBlackboardComponent& Blackboard) const
+{
+	return FBlackboard::GetValue<UBlackboardKeyType_Vector>(Blackboard, Key, KeyId, DefaultValue);
+}
+
+FVector FValueOrBBKey_Vector::GetValue(const UBlackboardComponent* Blackboard) const
+{
+	return Blackboard ? GetValue(*Blackboard) : DefaultValue;
+}
+
+FVector FValueOrBBKey_Vector::GetValue(const UBehaviorTreeComponent& BehaviorComp) const
+{
+	return FBlackboard::GetValue<UBlackboardKeyType_Vector>(BehaviorComp, Key, KeyId, DefaultValue);
+}
+
+FVector FValueOrBBKey_Vector::GetValue(const UBehaviorTreeComponent* BehaviorComp) const
+{
+	return BehaviorComp ? GetValue(*BehaviorComp) : DefaultValue;
+}
+
+#if WITH_EDITOR
+bool FValueOrBBKey_Bool::IsCompatibleType(const UBlackboardKeyType* KeyType) const
+{
+	return KeyType && KeyType->GetClass() == UBlackboardKeyType_Bool::StaticClass();
+}
+
+bool FValueOrBBKey_Class::IsCompatibleType(const UBlackboardKeyType* KeyType) const
+{
+	if (const UBlackboardKeyType_Class* ClassKey = Cast<UBlackboardKeyType_Class>(KeyType))
+	{
+		if (BaseClass && ClassKey->BaseClass)
+		{
+			return ClassKey->BaseClass->IsChildOf(BaseClass);
+		}
+	}
+	return false;
+}
+
+bool FValueOrBBKey_Enum::IsCompatibleType(const UBlackboardKeyType* KeyType) const
+{
+	if (const UBlackboardKeyType_Enum* EnumKey = Cast<UBlackboardKeyType_Enum>(KeyType))
+	{
+		if (EnumType && EnumKey->EnumType)
+		{
+			return EnumKey->EnumType == EnumType;
+		}
+	}
+	return false;
+}
+
+bool FValueOrBBKey_Float::IsCompatibleType(const UBlackboardKeyType* KeyType) const
+{
+	return KeyType && KeyType->GetClass() == UBlackboardKeyType_Float::StaticClass();
+}
+
+bool FValueOrBBKey_Int32::IsCompatibleType(const UBlackboardKeyType* KeyType) const
+{
+	return KeyType && KeyType->GetClass() == UBlackboardKeyType_Int::StaticClass();
+}
+
+bool FValueOrBBKey_Name::IsCompatibleType(const UBlackboardKeyType* KeyType) const
+{
+	return KeyType && KeyType->GetClass() == UBlackboardKeyType_Name::StaticClass();
+}
+
+bool FValueOrBBKey_String::IsCompatibleType(const UBlackboardKeyType* KeyType) const
+{
+	return KeyType && KeyType->GetClass() == UBlackboardKeyType_String::StaticClass();
+}
+
+bool FValueOrBBKey_Object::IsCompatibleType(const UBlackboardKeyType* KeyType) const
+{
+	if (const UBlackboardKeyType_Object* ObjectKey = Cast<UBlackboardKeyType_Object>(KeyType))
+	{
+		if (BaseClass && ObjectKey->BaseClass)
+		{
+			return ObjectKey->BaseClass->IsChildOf(BaseClass);
+		}
+	}
+	return false;
+}
+
+bool FValueOrBBKey_Rotator::IsCompatibleType(const UBlackboardKeyType* KeyType) const
+{
+	return KeyType && KeyType->GetClass() == UBlackboardKeyType_Rotator::StaticClass();
+}
+
+bool FValueOrBBKey_Vector::IsCompatibleType(const UBlackboardKeyType* KeyType) const
+{
+	return KeyType && KeyType->GetClass() == UBlackboardKeyType_Vector::StaticClass();
+}
+#endif // WITH_EDITOR
diff --git a/Engine/Source/Runtime/AIModule/Private/BehaviorTree/ValueOrBBKeyBlueprintUtility.h b/Engine/Source/Runtime/AIModule/Private/BehaviorTree/ValueOrBBKeyBlueprintUtility.h
new file mode 100644
--- /dev/null
+++ b/Engine/Source/Runtime/AIModule/Private/BehaviorTree/ValueOrBBKeyBlueprintUtility.h
@@ -0,0 +1,54 @@
+// Copyright Epic Games, Inc. All Rights Reserved.
+
+#pragma once
+
+#include "Kismet/BlueprintFunctionLibrary.h"
+
+#include "ValueOrBBKeyBlueprintUtility.generated.h"
+
+struct FValueOrBBKey_Bool;
+struct FValueOrBBKey_Class;
+struct FValueOrBBKey_Enum;
+struct FValueOrBBKey_Float;
+struct FValueOrBBKey_Int32;
+struct FValueOrBBKey_Name;
+struct FValueOrBBKey_String;
+struct FValueOrBBKey_Object;
+struct FValueOrBBKey_Rotator;
+struct FValueOrBBKey_Vector;
+
+UCLASS()
+class UValueOrBBKeyBlueprintUtility : public UBlueprintFunctionLibrary
+{
+	GENERATED_BODY()
+public:
+	UFUNCTION(BlueprintPure, Category = Blackboard)
+	static bool GetBool(const FValueOrBBKey_Bool& Value, const UBehaviorTreeComponent* BehaviorTreeComp);
+
+	UFUNCTION(BlueprintPure, Category = Blackboard)
+	static UClass* GetClass(const FValueOrBBKey_Class& Value, const UBehaviorTreeComponent* BehaviorTreeComp);
+
+	UFUNCTION(BlueprintPure, Category = Blackboard)
+	static uint8 GetEnum(const FValueOrBBKey_Enum& Value, const UBehaviorTreeComponent* BehaviorTreeComp);
+
+	UFUNCTION(BlueprintPure, Category = Blackboard)
+	static float GetFloat(const FValueOrBBKey_Float& Value, const UBehaviorTreeComponent* BehaviorTreeComp);
+
+	UFUNCTION(BlueprintPure, Category = Blackboard)
+	static int32 GetInt32(const FValueOrBBKey_Int32& Value, const UBehaviorTreeComponent* BehaviorTreeComp);
+
+	UFUNCTION(BlueprintPure, Category = Blackboard)
+	static FName GetName(const FValueOrBBKey_Name& Value, const UBehaviorTreeComponent* BehaviorTreeComp);
+
+	UFUNCTION(BlueprintPure, Category = Blackboard)
+	static FString GetString(const FValueOrBBKey_String& Value, const UBehaviorTreeComponent* BehaviorTreeComp);
+
+	UFUNCTION(BlueprintPure, Category = Blackboard)
+	static UObject* GetObject(const FValueOrBBKey_Object& Value, const UBehaviorTreeComponent* BehaviorTreeComp);
+
+	UFUNCTION(BlueprintPure, Category = Blackboard)
+	static FRotator GetRotator(const FValueOrBBKey_Rotator& Value, const UBehaviorTreeComponent* BehaviorTreeComp);
+
+	UFUNCTION(BlueprintPure, Category = Blackboard)
+	static FVector GetVector(const FValueOrBBKey_Vector& Value, const UBehaviorTreeComponent* BehaviorTreeComp);
+};
\ No newline at end of file
diff --git a/Engine/Source/Runtime/AIModule/Private/BehaviorTree/ValueOrKeyBlueprintUtility.cpp b/Engine/Source/Runtime/AIModule/Private/BehaviorTree/ValueOrKeyBlueprintUtility.cpp
new file mode 100644
--- /dev/null
+++ b/Engine/Source/Runtime/AIModule/Private/BehaviorTree/ValueOrKeyBlueprintUtility.cpp
@@ -0,0 +1,55 @@
+// Copyright Epic Games, Inc. All Rights Reserved.
+
+#include "BehaviorTree/ValueOrBBKeyBlueprintUtility.h"
+
+#include "BehaviorTree/ValueOrBBKey.h"
+
+bool UValueOrBBKeyBlueprintUtility::GetBool(const FValueOrBBKey_Bool& Value, const UBehaviorTreeComponent* BehaviorTreeComp)
+{
+	return Value.GetValue(BehaviorTreeComp);
+}
+
+UClass* UValueOrBBKeyBlueprintUtility::GetClass(const FValueOrBBKey_Class& Value, const UBehaviorTreeComponent* BehaviorTreeComp)
+{
+	return Value.GetValue(BehaviorTreeComp);
+}
+
+uint8 UValueOrBBKeyBlueprintUtility::GetEnum(const FValueOrBBKey_Enum& Value, const UBehaviorTreeComponent* BehaviorTreeComp)
+{
+	return Value.GetValue(BehaviorTreeComp);
+}
+
+float UValueOrBBKeyBlueprintUtility::GetFloat(const FValueOrBBKey_Float& Value, const UBehaviorTreeComponent* BehaviorTreeComp)
+{
+	return Value.GetValue(BehaviorTreeComp);
+}
+
+int32 UValueOrBBKeyBlueprintUtility::GetInt32(const FValueOrBBKey_Int32& Value, const UBehaviorTreeComponent* BehaviorTreeComp)
+{
+	return Value.GetValue(BehaviorTreeComp);
+}
+
+FName UValueOrBBKeyBlueprintUtility::GetName(const FValueOrBBKey_Name& Value, const UBehaviorTreeComponent* BehaviorTreeComp)
+{
+	return Value.GetValue(BehaviorTreeComp);
+}
+
+FString UValueOrBBKeyBlueprintUtility::GetString(const FValueOrBBKey_String& Value, const UBehaviorTreeComponent* BehaviorTreeComp)
+{
+	return Value.GetValue(BehaviorTreeComp);
+}
+
+UObject* UValueOrBBKeyBlueprintUtility::GetObject(const FValueOrBBKey_Object& Value, const UBehaviorTreeComponent* BehaviorTreeComp)
+{
+	return Value.GetValue(BehaviorTreeComp);
+}
+
+FRotator UValueOrBBKeyBlueprintUtility::GetRotator(const FValueOrBBKey_Rotator& Value, const UBehaviorTreeComponent* BehaviorTreeComp)
+{
+	return Value.GetValue(BehaviorTreeComp);
+}
+
+FVector UValueOrBBKeyBlueprintUtility::GetVector(const FValueOrBBKey_Vector& Value, const UBehaviorTreeComponent* BehaviorTreeComp)
+{
+	return Value.GetValue(BehaviorTreeComp);
+}
