Add a browse to asset button on ValueOrBBKey Class and object
#rnx

diff --git a/Engine/Source/Editor/BehaviorTreeEditor/Classes/ValueOrBBKeyDetails.h b/Engine/Source/Editor/BehaviorTreeEditor/Classes/ValueOrBBKeyDetails.h
--- a/Engine/Source/Editor/BehaviorTreeEditor/Classes/ValueOrBBKeyDetails.h
+++ b/Engine/Source/Editor/BehaviorTreeEditor/Classes/ValueOrBBKeyDetails.h
@@ -67,6 +67,7 @@ protected:
 	void OnBaseClassChanged();
 	void OnSetClass(const UClass* NewClass);
 	const UClass* OnGetSelectedClass() const;
+	void BrowseToClass() const;
 };
 
 class FValueOrBBKeyDetails_Enum : public FValueOrBBKeyDetails
@@ -110,6 +111,7 @@ protected:
 	void OnBaseClassChanged();
 	void OnObjectChanged(const FAssetData& AssetData);
 	FString OnGetObjectPath() const;
+	void BrowseToObject() const;
 };
 
 class FValueOrBBKeyDetails_Struct : public FValueOrBBKeyDetails
diff --git a/Engine/Source/Editor/BehaviorTreeEditor/Private/DetailCustomizations/ValueOrBBKeyDetails.cpp b/Engine/Source/Editor/BehaviorTreeEditor/Private/DetailCustomizations/ValueOrBBKeyDetails.cpp
--- a/Engine/Source/Editor/BehaviorTreeEditor/Private/DetailCustomizations/ValueOrBBKeyDetails.cpp
+++ b/Engine/Source/Editor/BehaviorTreeEditor/Private/DetailCustomizations/ValueOrBBKeyDetails.cpp
@@ -12,6 +12,7 @@
 #include "Delegates/Delegate.h"
 #include "DetailLayoutBuilder.h"
 #include "DetailWidgetRow.h"
+#include "Editor.h"
 #include "Fonts/SlateFontInfo.h"
 #include "Framework/Commands/UIAction.h"
 #include "Framework/MultiBox/MultiBoxBuilder.h"
@@ -157,12 +158,24 @@ TSharedRef<SWidget> FValueOrBBKeyDetails_Class::CreateDefaultValueWidget()
 	{
 		if (const UClass* BaseClass = static_cast<const FValueOrBBKey_Class*>(DataPtr)->BaseClass)
 		{
-			return SNew(SClassPropertyEntryBox)
-				.MetaClass(BaseClass)
-				.AllowNone(true)
-				.AllowAbstract(true)
-				.OnSetClass(this, &FValueOrBBKeyDetails_Class::OnSetClass)
-				.SelectedClass(this, &FValueOrBBKeyDetails_Class::OnGetSelectedClass);
+			return SNew(SHorizontalBox)
+				+SHorizontalBox::Slot()
+				[
+					SNew(SClassPropertyEntryBox)
+					.MetaClass(BaseClass)
+					.AllowNone(true)
+					.AllowAbstract(true)
+					.OnSetClass(this, &FValueOrBBKeyDetails_Class::OnSetClass)
+					.SelectedClass(this, &FValueOrBBKeyDetails_Class::OnGetSelectedClass)
+				]
+				+SHorizontalBox::Slot()
+				.AutoWidth()
+				.HAlign(HAlign_Center)
+				.VAlign(VAlign_Center)
+				.Padding(2.0f, 1.0f)
+				[
+					PropertyCustomizationHelpers::MakeBrowseButton(FSimpleDelegate::CreateSP(this, &FValueOrBBKeyDetails_Class::BrowseToClass))
+				];
 		}
 	}
 
@@ -235,11 +248,23 @@ TSharedRef<SWidget> FValueOrBBKeyDetails_Object::CreateDefaultValueWidget()
 	{
 		if (UClass* BaseClass = static_cast<const FValueOrBBKey_Object*>(DataPtr)->BaseClass)
 		{
-			return SNew(SObjectPropertyEntryBox)
-				.AllowedClass(BaseClass)
-				.AllowClear(true)
-				.OnObjectChanged(this, &FValueOrBBKeyDetails_Object::OnObjectChanged)
-				.ObjectPath(this, &FValueOrBBKeyDetails_Object::OnGetObjectPath);
+			return SNew(SHorizontalBox)
+					+ SHorizontalBox::Slot()
+					[
+						SNew(SObjectPropertyEntryBox)
+						.AllowedClass(BaseClass)
+						.AllowClear(true)
+						.OnObjectChanged(this, &FValueOrBBKeyDetails_Object::OnObjectChanged)
+						.ObjectPath(this, &FValueOrBBKeyDetails_Object::OnGetObjectPath)
+					]
+					+ SHorizontalBox::Slot()
+					 .AutoWidth()
+					 .HAlign(HAlign_Center)
+					 .VAlign(VAlign_Center)
+					 .Padding(2.0f, 1.0f)
+					[
+						PropertyCustomizationHelpers::MakeBrowseButton(FSimpleDelegate::CreateSP(this, &FValueOrBBKeyDetails_Object::BrowseToObject))
+					];
 		}
 	}
 	return FValueOrBBKeyDetails::CreateDefaultValueWidget();
@@ -456,6 +481,16 @@ const UClass* FValueOrBBKeyDetails_Class::OnGetSelectedClass() const
 	return Cast<UClass>(DefaultValue);
 }
 
+void FValueOrBBKeyDetails_Class::BrowseToClass() const
+{
+	UObject* DefaultValue = nullptr;
+	DefaultValueProperty->GetValue(DefaultValue);
+	if (DefaultValue)
+	{
+		GEditor->SyncBrowserToObject(DefaultValue);
+	}
+}
+
 void FValueOrBBKeyDetails_Enum::OnEnumTypeChanged()
 {
 	ValidateData();
@@ -515,4 +550,14 @@ FString FValueOrBBKeyDetails_Object::OnGetObjectPath() const
 	DefaultValueProperty->GetValue(DefaultValue);
 	return DefaultValue ? DefaultValue->GetPathName() : FString();
 }
+
+void FValueOrBBKeyDetails_Object::BrowseToObject() const
+{
+	UObject* DefaultValue = nullptr;
+	DefaultValueProperty->GetValue(DefaultValue);
+	if (DefaultValue)
+	{
+		GEditor->SyncBrowserToObject(DefaultValue);
+	}
+}
 #undef LOCTEXT_NAMESPACE
