Fix ValueOrBBKey details for vector, rotator and gameplay tag
#rnx

diff --git a/Engine/Plugins/Experimental/GameplayBehaviors/Source/GameplayBehaviorsEditorModule/Private/GameplayBehaviorsEditorModule.cpp b/Engine/Plugins/Experimental/GameplayBehaviors/Source/GameplayBehaviorsEditorModule/Private/GameplayBehaviorsEditorModule.cpp
--- a/Engine/Plugins/Experimental/GameplayBehaviors/Source/GameplayBehaviorsEditorModule/Private/GameplayBehaviorsEditorModule.cpp
+++ b/Engine/Plugins/Experimental/GameplayBehaviors/Source/GameplayBehaviorsEditorModule/Private/GameplayBehaviorsEditorModule.cpp
@@ -22,7 +22,7 @@ void FGameplayBehaviorsEditorModule::StartupModule()
 	FGameplayBehaviorsEditorStyle::Get();
 
 	FPropertyEditorModule& PropertyModule = FModuleManager::LoadModuleChecked<FPropertyEditorModule>("PropertyEditor");
-	PropertyModule.RegisterCustomPropertyTypeLayout("ValueOrBBKey_GameplayTag", FOnGetPropertyTypeCustomizationInstance::CreateStatic(&FValueOrBBKeyDetails::MakeInstance));
+	PropertyModule.RegisterCustomPropertyTypeLayout("ValueOrBBKey_GameplayTagContainer", FOnGetPropertyTypeCustomizationInstance::CreateStatic(&FValueOrBBKeyDetails_WithChild::MakeInstance));
 	PropertyModule.NotifyCustomizationModuleChanged();
 }
 
@@ -31,7 +31,7 @@ void FGameplayBehaviorsEditorModule::ShutdownModule()
 	FGameplayBehaviorsEditorStyle::Shutdown();
 
 	FPropertyEditorModule& PropertyModule = FModuleManager::LoadModuleChecked<FPropertyEditorModule>("PropertyEditor");
-	PropertyModule.UnregisterCustomPropertyTypeLayout("ValueOrBBKey_GameplayTag");
+	PropertyModule.UnregisterCustomPropertyTypeLayout("ValueOrBBKey_GameplayTagContainer");
 	PropertyModule.NotifyCustomizationModuleChanged();
 }
 
diff --git a/Engine/Source/Editor/BehaviorTreeEditor/Classes/ValueOrBBKeyDetails.h b/Engine/Source/Editor/BehaviorTreeEditor/Classes/ValueOrBBKeyDetails.h
--- a/Engine/Source/Editor/BehaviorTreeEditor/Classes/ValueOrBBKeyDetails.h
+++ b/Engine/Source/Editor/BehaviorTreeEditor/Classes/ValueOrBBKeyDetails.h
@@ -125,4 +125,12 @@ public:
 
 protected:
 	TSharedPtr<IPropertyHandle> EditDefaultsOnlyProperty;
+};
+
+class BEHAVIORTREEEDITOR_API FValueOrBBKeyDetails_WithChild : public FValueOrBBKeyDetails
+{
+public:
+	/** Makes a new instance of this detail layout class for a specific detail view requesting it */
+	static TSharedRef<IPropertyTypeCustomization> MakeInstance();
+	virtual void CustomizeChildren(TSharedRef<class IPropertyHandle> StructPropertyHandle, class IDetailChildrenBuilder& StructBuilder, IPropertyTypeCustomizationUtils& StructCustomizationUtils) override;
 };
\ No newline at end of file
diff --git a/Engine/Source/Editor/BehaviorTreeEditor/Private/BehaviorTreeEditorModule.cpp b/Engine/Source/Editor/BehaviorTreeEditor/Private/BehaviorTreeEditorModule.cpp
--- a/Engine/Source/Editor/BehaviorTreeEditor/Private/BehaviorTreeEditorModule.cpp
+++ b/Engine/Source/Editor/BehaviorTreeEditor/Private/BehaviorTreeEditorModule.cpp
@@ -68,8 +68,8 @@ void FBehaviorTreeEditorModule::StartupModule()
 	PropertyModule.RegisterCustomPropertyTypeLayout("ValueOrBBKey_Name", FOnGetPropertyTypeCustomizationInstance::CreateStatic(&FValueOrBBKeyDetails::MakeInstance));
 	PropertyModule.RegisterCustomPropertyTypeLayout("ValueOrBBKey_String", FOnGetPropertyTypeCustomizationInstance::CreateStatic(&FValueOrBBKeyDetails::MakeInstance));
 	PropertyModule.RegisterCustomPropertyTypeLayout("ValueOrBBKey_Object", FOnGetPropertyTypeCustomizationInstance::CreateStatic(&FValueOrBBKeyDetails_Object::MakeInstance));
-	PropertyModule.RegisterCustomPropertyTypeLayout("ValueOrBBKey_Rotator", FOnGetPropertyTypeCustomizationInstance::CreateStatic(&FValueOrBBKeyDetails::MakeInstance));
-	PropertyModule.RegisterCustomPropertyTypeLayout("ValueOrBBKey_Vector", FOnGetPropertyTypeCustomizationInstance::CreateStatic(&FValueOrBBKeyDetails::MakeInstance));
+	PropertyModule.RegisterCustomPropertyTypeLayout("ValueOrBBKey_Rotator", FOnGetPropertyTypeCustomizationInstance::CreateStatic(&FValueOrBBKeyDetails_WithChild::MakeInstance));
+	PropertyModule.RegisterCustomPropertyTypeLayout("ValueOrBBKey_Vector", FOnGetPropertyTypeCustomizationInstance::CreateStatic(&FValueOrBBKeyDetails_WithChild::MakeInstance));
 	PropertyModule.RegisterCustomPropertyTypeLayout("ValueOrBBKey_Struct", FOnGetPropertyTypeCustomizationInstance::CreateStatic(&FValueOrBBKeyDetails_Struct::MakeInstance));
 	PropertyModule.NotifyCustomizationModuleChanged();
 }
diff --git a/Engine/Source/Editor/BehaviorTreeEditor/Private/DetailCustomizations/ValueOrBBKeyDetails.cpp b/Engine/Source/Editor/BehaviorTreeEditor/Private/DetailCustomizations/ValueOrBBKeyDetails.cpp
--- a/Engine/Source/Editor/BehaviorTreeEditor/Private/DetailCustomizations/ValueOrBBKeyDetails.cpp
+++ b/Engine/Source/Editor/BehaviorTreeEditor/Private/DetailCustomizations/ValueOrBBKeyDetails.cpp
@@ -65,6 +65,11 @@ TSharedRef<IPropertyTypeCustomization> FValueOrBBKeyDetails_Struct::MakeInstance
 	return MakeShareable(new FValueOrBBKeyDetails_Struct);
 }
 
+TSharedRef<IPropertyTypeCustomization> FValueOrBBKeyDetails_WithChild::MakeInstance()
+{
+	return MakeShareable(new FValueOrBBKeyDetails_WithChild);
+}
+
 BEGIN_SLATE_FUNCTION_BUILD_OPTIMIZATION
 void FValueOrBBKeyDetails::CustomizeHeader(TSharedRef<IPropertyHandle> StructPropertyHandle, FDetailWidgetRow& HeaderRow, IPropertyTypeCustomizationUtils& StructCustomizationUtils)
 {
@@ -290,6 +295,13 @@ void FValueOrBBKeyDetails_Struct::CustomizeChildren(TSharedRef<class IPropertyHa
 	.IsEnabled(TAttribute<bool>(this, &FValueOrBBKeyDetails::CanEditDefaultValue));
 }
 
+void FValueOrBBKeyDetails_WithChild::CustomizeChildren(TSharedRef<class IPropertyHandle> StructPropertyHandle, class IDetailChildrenBuilder& StructBuilder, IPropertyTypeCustomizationUtils& StructCustomizationUtils)
+{
+	FValueOrBBKeyDetails::CustomizeChildren(StructPropertyHandle, StructBuilder, StructCustomizationUtils);
+	StructBuilder.AddProperty(DefaultValueProperty.ToSharedRef())
+	.IsEnabled(TAttribute<bool>(this, &FValueOrBBKeyDetails::CanEditDefaultValue));
+}
+
 END_SLATE_FUNCTION_BUILD_OPTIMIZATION
 
 bool FValueOrBBKeyDetails::HasAccessToBlackboard() const
@@ -560,4 +572,4 @@ void FValueOrBBKeyDetails_Object::BrowseToObject() const
 		GEditor->SyncBrowserToObject(DefaultValue);
 	}
 }
-#undef LOCTEXT_NAMESPACE
+#undef LOCTEXT_NAMESPACE
\ No newline at end of file
