Add GetAssociatedAsset to BTNode so that BTNode can define what asset should be opened when they are double clicked.

diff --git a/Engine/Plugins/Runtime/GameplayStateTree/Source/GameplayStateTreeModule/Private/BehaviorTree/Tasks/BTTask_RunDynamicStateTree.cpp b/Engine/Plugins/Runtime/GameplayStateTree/Source/GameplayStateTreeModule/Private/BehaviorTree/Tasks/BTTask_RunDynamicStateTree.cpp
--- a/Engine/Plugins/Runtime/GameplayStateTree/Source/GameplayStateTreeModule/Private/BehaviorTree/Tasks/BTTask_RunDynamicStateTree.cpp
+++ b/Engine/Plugins/Runtime/GameplayStateTree/Source/GameplayStateTreeModule/Private/BehaviorTree/Tasks/BTTask_RunDynamicStateTree.cpp
@@ -6,6 +6,7 @@
 #include "BehaviorTree/BTCompositeNode.h"
 #include "BehaviorTree/GameplayStateTreeBTUtils.h"
 #include "GameFramework/Actor.h"
+#include "Internationalization/Regex.h"
 #include "StateTreeExecutionContext.h"
 #include "VisualLogger/VisualLogger.h"
 
@@ -126,9 +127,28 @@ FString UBTTask_RunDynamicStateTree::GetStaticDescription() const
 void UBTTask_RunDynamicStateTree::DescribeRuntimeValues(const UBehaviorTreeComponent& OwnerComp, uint8* NodeMemory, EBTDescriptionVerbosity::Type Verbosity, TArray<FString>& Values) const
 {
 	Super::DescribeRuntimeValues(OwnerComp, NodeMemory, Verbosity, Values);
-	Values.Add(FString::Printf(TEXT("state tree: %s"), *GetNameSafe(StateTreeRef.GetStateTree())));
+	Values.Add(FString::Printf(TEXT("State tree: %s"), *GetNameSafe(StateTreeRef.GetStateTree())));
 	if (StateTreeRef.IsValid())
 	{
+		Values.Add(FString::Printf(TEXT("State tree path: %s"), *StateTreeRef.GetStateTree()->GetPathName()));
 		Values.Add(FString::Printf(TEXT("Interval: %.2f RandomDeviation: %.2f"), Interval, RandomDeviation));
 	}
-}
\ No newline at end of file
+}
+
+#if WITH_EDITOR
+const UObject* UBTTask_RunDynamicStateTree::GetAssociatedAsset(TOptional<FBehaviorTreeNodeDebugContext> DebugContext) const
+{
+	if (DebugContext.IsSet())
+	{
+		static const FRegexPattern RegexPatern(FString(TEXT("[\\n\\r].*state tree path:\\s*([^\\n\\r]*)")));
+		FRegexMatcher Matcher(RegexPatern, DebugContext->RuntimeDescription);
+		if (Matcher.FindNext())
+		{
+			// Capture group 0 is the whole subtree line, capture group 1 is the path.
+			const FString StateTreePath = Matcher.GetCaptureGroup(1);
+			return FindObject<UStateTree>(nullptr, *StateTreePath);
+		}
+	}
+	return StateTreeRef.GetStateTree();
+}
+#endif //WITH_EDITOR
\ No newline at end of file
diff --git a/Engine/Plugins/Runtime/GameplayStateTree/Source/GameplayStateTreeModule/Private/BehaviorTree/Tasks/BTTask_RunStateTree.cpp b/Engine/Plugins/Runtime/GameplayStateTree/Source/GameplayStateTreeModule/Private/BehaviorTree/Tasks/BTTask_RunStateTree.cpp
--- a/Engine/Plugins/Runtime/GameplayStateTree/Source/GameplayStateTreeModule/Private/BehaviorTree/Tasks/BTTask_RunStateTree.cpp
+++ b/Engine/Plugins/Runtime/GameplayStateTree/Source/GameplayStateTreeModule/Private/BehaviorTree/Tasks/BTTask_RunStateTree.cpp
@@ -4,6 +4,7 @@
 
 #include "BehaviorTree/GameplayStateTreeBTUtils.h"
 #include "Components/StateTreeAIComponentSchema.h"
+#include "StateTree.h"
 #include "StateTreeExecutionContext.h"
 
 UBTTask_RunStateTree::UBTTask_RunStateTree(const FObjectInitializer& ObjectInitializer /*= FObjectInitializer::Get()*/)
@@ -80,3 +81,10 @@ TSubclassOf<UStateTreeSchema> UBTTask_RunStateTree::GetSchema() const
 {
 	return UStateTreeAIComponentSchema::StaticClass();
 }
+
+#if WITH_EDITOR
+const UObject* UBTTask_RunStateTree::GetAssociatedAsset(TOptional<FBehaviorTreeNodeDebugContext> DebugContext) const
+{
+	return StateTreeRef.GetStateTree();
+}
+#endif // WITH_EDITOR
diff --git a/Engine/Plugins/Runtime/GameplayStateTree/Source/GameplayStateTreeModule/Public/BehaviorTree/Tasks/BTTask_RunDynamicStateTree.h b/Engine/Plugins/Runtime/GameplayStateTree/Source/GameplayStateTreeModule/Public/BehaviorTree/Tasks/BTTask_RunDynamicStateTree.h
--- a/Engine/Plugins/Runtime/GameplayStateTree/Source/GameplayStateTreeModule/Public/BehaviorTree/Tasks/BTTask_RunDynamicStateTree.h
+++ b/Engine/Plugins/Runtime/GameplayStateTree/Source/GameplayStateTreeModule/Public/BehaviorTree/Tasks/BTTask_RunDynamicStateTree.h
@@ -39,6 +39,10 @@ protected:
 
 	void SetStateTreeToRun(UBehaviorTreeComponent& OwnerComp, const FStateTreeReference& StateTreeToRun, const FSetContextDataDelegate& SetContextDelegate, float Interval, float RandomDeviation);
 
+#if WITH_EDITOR
+	GAMEPLAYSTATETREEMODULE_API const UObject* GetAssociatedAsset(TOptional<FBehaviorTreeNodeDebugContext> DebugContext) const override;
+#endif //WITH_EDITOR
+
 	UPROPERTY(Transient)
 	FStateTreeReference StateTreeRef;
 
diff --git a/Engine/Plugins/Runtime/GameplayStateTree/Source/GameplayStateTreeModule/Public/BehaviorTree/Tasks/BTTask_RunStateTree.h b/Engine/Plugins/Runtime/GameplayStateTree/Source/GameplayStateTreeModule/Public/BehaviorTree/Tasks/BTTask_RunStateTree.h
--- a/Engine/Plugins/Runtime/GameplayStateTree/Source/GameplayStateTreeModule/Public/BehaviorTree/Tasks/BTTask_RunStateTree.h
+++ b/Engine/Plugins/Runtime/GameplayStateTree/Source/GameplayStateTreeModule/Public/BehaviorTree/Tasks/BTTask_RunStateTree.h
@@ -31,6 +31,10 @@ protected:
 
 	virtual TSubclassOf<UStateTreeSchema> GetSchema() const override;
 
+#if WITH_EDITOR
+	GAMEPLAYSTATETREEMODULE_API virtual const UObject* GetAssociatedAsset(TOptional<FBehaviorTreeNodeDebugContext> DebugContext) const override;
+#endif //WITH_EDITOR
+
 	/** State tree that will be run when the task is selected. */
 	UPROPERTY(EditAnywhere, Category = Task, meta = (Schema = "/Script/GameplayStateTreeModule.StateTreeAIComponentSchema", SchemaCanBeOverriden))
 	FStateTreeReference StateTreeRef;
diff --git a/Engine/Source/Editor/BehaviorTreeEditor/Private/BehaviorTreeDebugger.cpp b/Engine/Source/Editor/BehaviorTreeEditor/Private/BehaviorTreeDebugger.cpp
--- a/Engine/Source/Editor/BehaviorTreeEditor/Private/BehaviorTreeDebugger.cpp
+++ b/Engine/Source/Editor/BehaviorTreeEditor/Private/BehaviorTreeDebugger.cpp
@@ -8,7 +8,6 @@
 #include "EdGraph/EdGraph.h"
 #include "Editor/EditorEngine.h"
 #include "Editor.h"
-#include "BehaviorTree/BTNode.h"
 #include "BehaviorTree/BTTaskNode.h"
 #include "BehaviorTree/BTAuxiliaryNode.h"
 #include "BehaviorTree/Tasks/BTTask_RunBehaviorDynamic.h"
@@ -1229,7 +1228,7 @@ void FBehaviorTreeDebugger::UpdateDebuggerViewOnTick()
 #endif
 }
 
-class UBehaviorTree* FBehaviorTreeDebugger::GetDynamicSubtreeTaskBehaviorTree(const UBTTask_RunBehaviorDynamic* Node) const
+TOptional<FBehaviorTreeNodeDebugContext> FBehaviorTreeDebugger::GetDebugContextForNode(const UBTNode* Node) const
 {
 	if (UBehaviorTreeComponent* TreeComp = TreeInstance.Get())
 	{
@@ -1241,12 +1240,13 @@ class UBehaviorTree* FBehaviorTreeDebugger::GetDynamicSubtreeTaskBehaviorTree(co
 				const FBehaviorTreeDebuggerInstance& DebugInstance = ActiveStep.InstanceStack[DebuggerInstanceIndex];
 				if (DebugInstance.RuntimeDesc.IsValidIndex(Node->GetExecutionIndex()))
 				{
-					return Node->GetBehaviorAssetFromRuntimeValue(DebugInstance.RuntimeDesc[Node->GetExecutionIndex()]);
+					return FBehaviorTreeNodeDebugContext{ *TreeComp, DebugInstance.RuntimeDesc[Node->GetExecutionIndex()]};
 				}
 			}
 		}
+		return FBehaviorTreeNodeDebugContext{ *TreeComp, FString() };
 	}
-	return nullptr;
+	return {};
 }
 
 FText FBehaviorTreeDebugger::FindValueForKey(const FName& InKeyName, bool bUseCurrentState) const
diff --git a/Engine/Source/Editor/BehaviorTreeEditor/Private/BehaviorTreeDebugger.h b/Engine/Source/Editor/BehaviorTreeEditor/Private/BehaviorTreeDebugger.h
--- a/Engine/Source/Editor/BehaviorTreeEditor/Private/BehaviorTreeDebugger.h
+++ b/Engine/Source/Editor/BehaviorTreeEditor/Private/BehaviorTreeDebugger.h
@@ -21,7 +21,9 @@ class FBehaviorTreeEditor;
 class UBehaviorTree;
 class UBehaviorTreeGraphNode;
 class UBlackboardData;
+class UBTNode;
 class UObject;
+struct FBehaviorTreeNodeDebugContext;
 
 class FBehaviorTreeDebugger : public FTickableGameObject
 {
@@ -87,7 +89,7 @@ public:
 	bool HasContinuousNextStep() const;
 	bool HasContinuousPrevStep() const;
 
-	UBehaviorTree* GetDynamicSubtreeTaskBehaviorTree(const UBTTask_RunBehaviorDynamic* Node) const;
+	TOptional<FBehaviorTreeNodeDebugContext> GetDebugContextForNode(const UBTNode* Node) const;
 
 	/**
 	 * Find a (display) value for a given key.
diff --git a/Engine/Source/Editor/BehaviorTreeEditor/Private/BehaviorTreeEditor.cpp b/Engine/Source/Editor/BehaviorTreeEditor/Private/BehaviorTreeEditor.cpp
--- a/Engine/Source/Editor/BehaviorTreeEditor/Private/BehaviorTreeEditor.cpp
+++ b/Engine/Source/Editor/BehaviorTreeEditor/Private/BehaviorTreeEditor.cpp
@@ -1255,50 +1255,24 @@ void FBehaviorTreeEditor::OnNodeDoubleClicked(class UEdGraphNode* Node)
 			}
 		}
 	}
-	else if (UBehaviorTreeGraphNode_Task* Task = Cast<UBehaviorTreeGraphNode_Task>(Node))
+	else if (MyNode)
 	{
-		UBehaviorTree* SubTreeToOpen = nullptr;
-		if (UBTTask_RunBehavior* SubtreeTask = Cast<UBTTask_RunBehavior>(Task->NodeInstance))
+		if (const UBTNode* NodeInstance = Cast<UBTNode>(MyNode->NodeInstance))
 		{
-			SubTreeToOpen = SubtreeTask->GetSubtreeAsset();
-		}
-		else if (UBTTask_RunBehaviorDynamic* DynamicSubTreeTask = Cast<UBTTask_RunBehaviorDynamic>(Task->NodeInstance))
-		{
-			if (Debugger)
-			{
-				SubTreeToOpen = Debugger->GetDynamicSubtreeTaskBehaviorTree(DynamicSubTreeTask);
-			}
-			
-			if (!SubTreeToOpen)
-			{
-				SubTreeToOpen = DynamicSubTreeTask->GetDefaultBehaviorAsset();
-			}
-		}
-
-		if (SubTreeToOpen)
-		{
-			GEditor->GetEditorSubsystem<UAssetEditorSubsystem>()->OpenEditorForAsset(SubTreeToOpen);
-
-			IBehaviorTreeEditor* ChildNodeEditor = static_cast<IBehaviorTreeEditor*>(GEditor->GetEditorSubsystem<UAssetEditorSubsystem>()->FindEditorForAsset(SubTreeToOpen, true));
-			if (ChildNodeEditor)
+			const TOptional<FBehaviorTreeNodeDebugContext> DebugContext = Debugger ? Debugger->GetDebugContextForNode(NodeInstance) : TOptional<FBehaviorTreeNodeDebugContext>{};
+			if (const UObject* AssetToOpen = NodeInstance->GetAssociatedAsset(DebugContext))
 			{
-				ChildNodeEditor->InitializeDebuggerState(Debugger.Get());
+				GEditor->GetEditorSubsystem<UAssetEditorSubsystem>()->OpenEditorForAsset(AssetToOpen);
+				if (const UBehaviorTree* OpenedBehaviorTree = Cast<UBehaviorTree>(AssetToOpen))
+				{
+					if (IBehaviorTreeEditor* ChildNodeEditor = static_cast<IBehaviorTreeEditor*>(GEditor->GetEditorSubsystem<UAssetEditorSubsystem>()->FindEditorForAsset(const_cast<UBehaviorTree*>(OpenedBehaviorTree), true)))
+					{
+						ChildNodeEditor->InitializeDebuggerState(Debugger.Get());
+					}
+				}
 			}
 		}
 	}
-
-	if (MyNode && MyNode->NodeInstance &&
-		MyNode->NodeInstance->GetClass()->HasAnyClassFlags(CLASS_CompiledFromBlueprint))
-	{
-		UClass* NodeClass = MyNode->NodeInstance->GetClass();
-		UPackage* Pkg = NodeClass->GetOuterUPackage();
-		FString ClassName = NodeClass->GetName().LeftChop(2);
-		UBlueprint* BlueprintOb = FindObject<UBlueprint>(Pkg, *ClassName);
-		if(BlueprintOb)
-		{
-			GEditor->GetEditorSubsystem<UAssetEditorSubsystem>()->OpenEditorForAsset(BlueprintOb);
-		}
-	}
 }
 
 void FBehaviorTreeEditor::OnAddInputPin()
diff --git a/Engine/Source/Runtime/AIModule/Classes/BehaviorTree/BTNode.h b/Engine/Source/Runtime/AIModule/Classes/BehaviorTree/BTNode.h
--- a/Engine/Source/Runtime/AIModule/Classes/BehaviorTree/BTNode.h
+++ b/Engine/Source/Runtime/AIModule/Classes/BehaviorTree/BTNode.h
@@ -35,6 +35,14 @@ struct FBTInstancedNodeMemory
 	int32 NodeIdx;
 };
 
+struct FBehaviorTreeNodeDebugContext
+{
+	/** Behavior Tree component of the behavior tree being debugged. */
+	UBehaviorTreeComponent& OwnerComponent;
+	/** Combined string of the values returned by DescribeRuntimeValues of the node at the current debug execution step. */
+	FString RuntimeDescription;
+};
+
 UCLASS(Abstract,config=Game, MinimalAPI)
 class UBTNode : public UObject, public IGameplayTaskOwnerInterface
 {
@@ -171,6 +179,9 @@ class UBTNode : public UObject, public IGameplayTaskOwnerInterface
 	virtual void OnNodeCreated() {}
 
 	AIMODULE_API void PreSave(FObjectPreSaveContext SaveContext) override;
+
+	/** @return the external asset associated with this node that should be open on double click. If set, DebugContext contains the info of the node in the current debug step. */
+	AIMODULE_API virtual const UObject* GetAssociatedAsset(TOptional<FBehaviorTreeNodeDebugContext> DebugContext) const;
 #endif
 
 	/** Gets called only for instanced nodes(bCreateNodeInstance == true). In practive overridden by BP-implemented BT nodes */
diff --git a/Engine/Source/Runtime/AIModule/Classes/BehaviorTree/Tasks/BTTask_RunBehavior.h b/Engine/Source/Runtime/AIModule/Classes/BehaviorTree/Tasks/BTTask_RunBehavior.h
--- a/Engine/Source/Runtime/AIModule/Classes/BehaviorTree/Tasks/BTTask_RunBehavior.h
+++ b/Engine/Source/Runtime/AIModule/Classes/BehaviorTree/Tasks/BTTask_RunBehavior.h
@@ -29,6 +29,7 @@ class UBTTask_RunBehavior : public UBTTaskNode
 
 #if WITH_EDITOR
 	AIMODULE_API virtual FName GetNodeIconName() const override;
+	AIMODULE_API virtual const UObject* GetAssociatedAsset(TOptional<FBehaviorTreeNodeDebugContext> DebugContext) const override;
 #endif // WITH_EDITOR
 
 	/** @returns number of injected nodes */
diff --git a/Engine/Source/Runtime/AIModule/Classes/BehaviorTree/Tasks/BTTask_RunBehaviorDynamic.h b/Engine/Source/Runtime/AIModule/Classes/BehaviorTree/Tasks/BTTask_RunBehaviorDynamic.h
--- a/Engine/Source/Runtime/AIModule/Classes/BehaviorTree/Tasks/BTTask_RunBehaviorDynamic.h
+++ b/Engine/Source/Runtime/AIModule/Classes/BehaviorTree/Tasks/BTTask_RunBehaviorDynamic.h
@@ -29,7 +29,11 @@ class UBTTask_RunBehaviorDynamic : public UBTTaskNode
 
 #if WITH_EDITOR
 	AIMODULE_API virtual FName GetNodeIconName() const override;
+	AIMODULE_API virtual const UObject* GetAssociatedAsset(TOptional<FBehaviorTreeNodeDebugContext> DebugContext) const override;
+
+	UE_DEPRECATED(5.6, "GetBehaviorAssetFromRuntimeValue has been replaced with the more generic GetAssociatedAsset.")
 	AIMODULE_API UBehaviorTree* GetBehaviorAssetFromRuntimeValue(const FString& RuntimeValue) const;
+
 #endif // WITH_EDITOR
 
 	AIMODULE_API bool HasMatchingTag(const FGameplayTag& Tag) const;
diff --git a/Engine/Source/Runtime/AIModule/Private/BehaviorTree/BTNode.cpp b/Engine/Source/Runtime/AIModule/Private/BehaviorTree/BTNode.cpp
--- a/Engine/Source/Runtime/AIModule/Private/BehaviorTree/BTNode.cpp
+++ b/Engine/Source/Runtime/AIModule/Private/BehaviorTree/BTNode.cpp
@@ -254,7 +254,13 @@ void UBTNode::PreSave(FObjectPreSaveContext SaveContext)
 	Super::PreSave(SaveContext);
 }
 
-#endif
+const UObject* UBTNode::GetAssociatedAsset(TOptional<FBehaviorTreeNodeDebugContext> DebugContext) const
+{
+	// By default we return the blueprint the node is derived from. If the node is not a blueprint it will return nullptr.
+	return UBlueprint::GetBlueprintFromClass(GetClass());
+}
+
+#endif // WITH_EDITOR
 
 UGameplayTasksComponent* UBTNode::GetGameplayTasksComponent(const UGameplayTask& Task) const
 {
diff --git a/Engine/Source/Runtime/AIModule/Private/BehaviorTree/Tasks/BTTask_RunBehavior.cpp b/Engine/Source/Runtime/AIModule/Private/BehaviorTree/Tasks/BTTask_RunBehavior.cpp
--- a/Engine/Source/Runtime/AIModule/Private/BehaviorTree/Tasks/BTTask_RunBehavior.cpp
+++ b/Engine/Source/Runtime/AIModule/Private/BehaviorTree/Tasks/BTTask_RunBehavior.cpp
@@ -61,5 +61,10 @@ FName UBTTask_RunBehavior::GetNodeIconName() const
 	return FName("BTEditor.Graph.BTNode.Task.RunBehavior.Icon");
 }
 
+const UObject* UBTTask_RunBehavior::GetAssociatedAsset(TOptional<FBehaviorTreeNodeDebugContext> DebugContext) const
+{
+	return BehaviorAsset;
+}
+
 #endif	// WITH_EDITOR
 
diff --git a/Engine/Source/Runtime/AIModule/Private/BehaviorTree/Tasks/BTTask_RunBehaviorDynamic.cpp b/Engine/Source/Runtime/AIModule/Private/BehaviorTree/Tasks/BTTask_RunBehaviorDynamic.cpp
--- a/Engine/Source/Runtime/AIModule/Private/BehaviorTree/Tasks/BTTask_RunBehaviorDynamic.cpp
+++ b/Engine/Source/Runtime/AIModule/Private/BehaviorTree/Tasks/BTTask_RunBehaviorDynamic.cpp
@@ -78,6 +78,27 @@ bool UBTTask_RunBehaviorDynamic::SetBehaviorAsset(UBehaviorTree* NewBehaviorAsse
 }
 
 #if WITH_EDITOR
+FName UBTTask_RunBehaviorDynamic::GetNodeIconName() const
+{
+	return FName("BTEditor.Graph.BTNode.Task.RunBehavior.Icon");
+}
+
+const UObject* UBTTask_RunBehaviorDynamic::GetAssociatedAsset(TOptional<FBehaviorTreeNodeDebugContext> DebugContext) const
+{
+	if (DebugContext.IsSet())
+	{
+		static const FRegexPattern RegexPatern(FString(TEXT("[\\n\\r].*subtree path:\\s*([^\\n\\r]*)")));
+		FRegexMatcher Matcher(RegexPatern, DebugContext->RuntimeDescription);
+		if (Matcher.FindNext())
+		{
+			// Capture group 0 is the whole subtree line, capture group 1 is the path.
+			const FString SubtreePath = Matcher.GetCaptureGroup(1);
+			return FindObject<UBehaviorTree>(nullptr, *SubtreePath);
+		}
+	}
+	return GetDefaultBehaviorAsset();
+}
+
 UBehaviorTree* UBTTask_RunBehaviorDynamic::GetBehaviorAssetFromRuntimeValue(const FString& RuntimeValue) const
 {
 	static const FRegexPattern RegexPatern(FString(TEXT("[\\n\\r].*subtree path:\\s*([^\\n\\r]*)")));
@@ -91,10 +112,5 @@ UBehaviorTree* UBTTask_RunBehaviorDynamic::GetBehaviorAssetFromRuntimeValue(cons
 	return nullptr;
 }
 
-FName UBTTask_RunBehaviorDynamic::GetNodeIconName() const
-{
-	return FName("BTEditor.Graph.BTNode.Task.RunBehavior.Icon");
-}
-
 #endif	// WITH_EDITOR
 
