Blackboard editor: Add duplicate blackboard entry

#rnx
#tests Tested in editor, duplicate entries in blackboard

diff --git a/Engine/Source/Editor/BehaviorTreeEditor/Private/SBehaviorTreeBlackboardEditor.cpp b/Engine/Source/Editor/BehaviorTreeEditor/Private/SBehaviorTreeBlackboardEditor.cpp
--- a/Engine/Source/Editor/BehaviorTreeEditor/Private/SBehaviorTreeBlackboardEditor.cpp
+++ b/Engine/Source/Editor/BehaviorTreeEditor/Private/SBehaviorTreeBlackboardEditor.cpp
@@ -68,6 +68,12 @@ void SBehaviorTreeBlackboardEditor::Construct(const FArguments& InArgs, TSharedR
 		FCanExecuteAction::CreateSP(this, &SBehaviorTreeBlackboardEditor::CanDeleteEntry)
 		);
 
+	CommandList->MapAction(
+		FGenericCommands::Get().Duplicate,
+		FExecuteAction::CreateSP(this, &SBehaviorTreeBlackboardEditor::HandleDuplicateEntry),
+		FCanExecuteAction::CreateSP(this, &SBehaviorTreeBlackboardEditor::CanDuplicateEntry)
+		);
+
 	CommandList->MapAction(
 		FGenericCommands::Get().Rename,
 		FExecuteAction::CreateSP(this, &SBehaviorTreeBlackboardEditor::HandleRenameEntry),
@@ -92,9 +98,10 @@ void SBehaviorTreeBlackboardEditor::Construct(const FArguments& InArgs, TSharedR
 
 void SBehaviorTreeBlackboardEditor::FillContextMenu(FMenuBuilder& MenuBuilder) const
 {
-	if(!IsDebuggerActive() && HasSelectedItems())
+	if (!IsDebuggerActive() && HasSelectedItems())
 	{
 		MenuBuilder.AddMenuEntry(FBTBlackboardCommands::Get().DeleteEntry);
+		MenuBuilder.AddMenuEntry(FGenericCommands::Get().Duplicate, NAME_None, LOCTEXT("Duplicate", "Duplicate"), LOCTEXT("Duplicate_Tooltip", "Duplicates this blackboard entry.") );
 		MenuBuilder.AddMenuEntry(FGenericCommands::Get().Rename, NAME_None, LOCTEXT("Rename", "Rename"), LOCTEXT("Rename_Tooltip", "Renames this blackboard entry.") );
 	}
 }
@@ -130,11 +137,11 @@ void SBehaviorTreeBlackboardEditor::HandleDeleteEntry()
 		return;
 	}
 
-	if(!IsDebuggerActive())
+	if (!IsDebuggerActive())
 	{
 		bool bIsInherited = false;
 		FBlackboardEntry* BlackboardEntry = GetSelectedEntry(bIsInherited);
-		if(BlackboardEntry != nullptr && !bIsInherited)
+		if (BlackboardEntry != nullptr && !bIsInherited)
 		{
 			const FScopedTransaction Transaction(LOCTEXT("BlackboardEntryDeleteTransaction", "Delete Blackboard Entry"));
 			BlackboardDataPtr->SetFlags(RF_Transactional);
@@ -143,9 +150,9 @@ void SBehaviorTreeBlackboardEditor::HandleDeleteEntry()
 			FProperty* KeysProperty = FindFProperty<FProperty>(UBlackboardData::StaticClass(), GET_MEMBER_NAME_CHECKED(UBlackboardData, Keys));
 			BlackboardDataPtr->PreEditChange(KeysProperty);
 		
-			for(int32 ItemIndex = 0; ItemIndex < BlackboardDataPtr->Keys.Num(); ItemIndex++)
+			for (int32 ItemIndex = 0; ItemIndex < BlackboardDataPtr->Keys.Num(); ItemIndex++)
 			{
-				if(BlackboardEntry == &BlackboardDataPtr->Keys[ItemIndex])
+				if (BlackboardEntry == &BlackboardDataPtr->Keys[ItemIndex])
 				{
 					BlackboardDataPtr->Keys.RemoveAt(ItemIndex);
 					break;
@@ -153,13 +160,10 @@ void SBehaviorTreeBlackboardEditor::HandleDeleteEntry()
 			}
 
 			GraphActionMenu->RefreshAllActions(true);
-			UE::BehaviorTreeEditor::Delegates::OnBlackboardKeyChanged.Broadcast(*BlackboardDataPtr, nullptr);
+			UE::BehaviorTreeEditor::Delegates::OnBlackboardKeyChanged.Broadcast(*BlackboardDataPtr, /* FBlackboardEntry* */ nullptr);
 
 			// signal de-selection
-			if(OnEntrySelected.IsBound())
-			{
-				OnEntrySelected.Execute(nullptr, false);
-			}
+			OnEntrySelected.ExecuteIfBound( /* FBlackboardEntry* */ nullptr, false);
 
 			FPropertyChangedEvent PropertyChangedEvent(KeysProperty, EPropertyChangeType::ArrayRemove);
 			BlackboardDataPtr->PostEditChangeProperty(PropertyChangedEvent);
@@ -167,9 +171,50 @@ void SBehaviorTreeBlackboardEditor::HandleDeleteEntry()
 	}
 }
 
+void SBehaviorTreeBlackboardEditor::HandleDuplicateEntry()
+{
+	UBlackboardData* BlackboardDataPtr = BlackboardData.Get();
+	if (!BlackboardDataPtr)
+	{
+		UE_LOG(LogBlackboardEditor, Error, TEXT("Trying to duplicate an entry from a blackboard while no Blackboard Asset is set!"));
+		return;
+	}
+
+	if (!IsDebuggerActive())
+	{
+		bool bIsInherited = false;
+		FBlackboardEntry* BlackboardEntry = GetSelectedEntry(bIsInherited);
+		if (BlackboardEntry != nullptr)
+		{
+			const FScopedTransaction Transaction(LOCTEXT("BlackboardEntryDuplicateTransaction", "Duplicate Blackboard Entry"));
+			BlackboardDataPtr->SetFlags(RF_Transactional);
+			BlackboardDataPtr->Modify();
+
+			FProperty* KeysProperty = FindFProperty<FProperty>(UBlackboardData::StaticClass(), GET_MEMBER_NAME_CHECKED(UBlackboardData, Keys));
+			BlackboardDataPtr->PreEditChange(KeysProperty);
+
+			FBlackboardEntry NewEntry = *BlackboardEntry;
+			NewEntry.EntryName = FName(BlackboardEntry->EntryName.ToString() + TEXT("Copy"));
+			NewEntry.KeyType = DuplicateObject<UBlackboardKeyType>(BlackboardEntry->KeyType, BlackboardDataPtr);
+			BlackboardDataPtr->Keys.Add(NewEntry);
+
+			GraphActionMenu->RefreshAllActions(true);
+			UE::BehaviorTreeEditor::Delegates::OnBlackboardKeyChanged.Broadcast(*BlackboardDataPtr, &BlackboardDataPtr->Keys.Last());
+
+			FPropertyChangedEvent PropertyChangedEvent(KeysProperty, EPropertyChangeType::ArrayAdd);
+			BlackboardDataPtr->PostEditChangeProperty(PropertyChangedEvent);
+			
+			if (SetSelectedEntry(NewEntry.EntryName))
+			{
+				GraphActionMenu->OnRequestRenameOnActionNode();
+			}
+		}
+	}
+}
+
 void SBehaviorTreeBlackboardEditor::HandleRenameEntry()
 {
-	if(!IsDebuggerActive())
+	if (!IsDebuggerActive())
 	{
 		GraphActionMenu->OnRequestRenameOnActionNode();
 	}
@@ -179,11 +224,11 @@ bool SBehaviorTreeBlackboardEditor::CanDeleteEntry() const
 {
 	const bool bModeActive = OnIsBlackboardModeActive.IsBound() && OnIsBlackboardModeActive.Execute();
 
-	if(!IsDebuggerActive() && bModeActive)
+	if (!IsDebuggerActive() && bModeActive)
 	{
 		bool bIsInherited = false;
 		FBlackboardEntry* BlackboardEntry = GetSelectedEntry(bIsInherited);
-		if(BlackboardEntry != nullptr)
+		if (BlackboardEntry != nullptr)
 		{
 			return !bIsInherited;
 		}
@@ -192,15 +237,32 @@ bool SBehaviorTreeBlackboardEditor::CanDeleteEntry() const
 	return false;
 }
 
+bool SBehaviorTreeBlackboardEditor::CanDuplicateEntry() const
+{
+	const bool bModeActive = OnIsBlackboardModeActive.IsBound() && OnIsBlackboardModeActive.Execute();
+
+	if (!IsDebuggerActive() && bModeActive)
+	{
+		bool bIsInherited = false;
+		FBlackboardEntry* BlackboardEntry = GetSelectedEntry(bIsInherited);
+		if (BlackboardEntry != nullptr)
+		{
+			return true;
+		}
+	}
+
+	return false;
+}
+
 bool SBehaviorTreeBlackboardEditor::CanRenameEntry() const
 {
 	const bool bModeActive = OnIsBlackboardModeActive.IsBound() && OnIsBlackboardModeActive.Execute();
 
-	if(!IsDebuggerActive() && bModeActive)
+	if (!IsDebuggerActive() && bModeActive)
 	{
 		bool bIsInherited = false;
 		FBlackboardEntry* BlackboardEntry = GetSelectedEntry(bIsInherited);
-		if(BlackboardEntry != nullptr)
+		if (BlackboardEntry != nullptr)
 		{
 			return !bIsInherited;
 		}
@@ -215,7 +277,7 @@ public:
 
 	virtual bool IsClassAllowed(const FClassViewerInitializationOptions& InInitOptions, const UClass* InClass, TSharedRef< FClassViewerFilterFuncs > InFilterFuncs ) override
 	{
-		if(InClass != nullptr)
+		if (InClass != nullptr)
 		{
 			return !InClass->HasAnyClassFlags(CLASS_Abstract | CLASS_HideDropDown) &&
 				InClass->HasAnyClassFlags(CLASS_EditInlineNew) &&
@@ -283,14 +345,14 @@ void SBehaviorTreeBlackboardEditor::HandleKeyClassPicked(UClass* InClass)
 	int32 IndexSuffix = -1;
 	auto DuplicateFunction = [&](const FBlackboardEntry& Key)
 	{		
-		if(Key.EntryName.ToString() == NewKeyName)
+		if (Key.EntryName.ToString() == NewKeyName)
 		{
 			IndexSuffix = FMath::Max(0, IndexSuffix);
 		}
-		if(Key.EntryName.ToString().StartsWith(NewKeyName))
+		if (Key.EntryName.ToString().StartsWith(NewKeyName))
 		{
 			const FString ExistingSuffix = Key.EntryName.ToString().RightChop(NewKeyName.Len());
-			if(ExistingSuffix.IsNumeric())
+			if (ExistingSuffix.IsNumeric())
 			{
 				IndexSuffix = FMath::Max(FCString::Atoi(*ExistingSuffix) + 1, IndexSuffix);
 			}
@@ -298,10 +360,10 @@ void SBehaviorTreeBlackboardEditor::HandleKeyClassPicked(UClass* InClass)
 	};
 
 	// check for existing keys of the same name
-	for(const auto& Key : BlackboardDataPtr->Keys) { DuplicateFunction(Key); };
-	for(const auto& Key : BlackboardDataPtr->ParentKeys) { DuplicateFunction(Key); };
+	for (const auto& Key : BlackboardDataPtr->Keys) { DuplicateFunction(Key); };
+	for (const auto& Key : BlackboardDataPtr->ParentKeys) { DuplicateFunction(Key); };
 
-	if(IndexSuffix != -1)
+	if (IndexSuffix != -1)
 	{
 		NewKeyName += FString::Printf(TEXT("%d"), IndexSuffix);
 	}
@@ -333,7 +395,7 @@ void SBehaviorTreeBlackboardEditor::HandleKeyClassPicked(UClass* InClass)
 
 bool SBehaviorTreeBlackboardEditor::CanCreateNewEntry() const
 {
-	if(OnIsDebuggerReady.IsBound())
+	if (OnIsDebuggerReady.IsBound())
 	{
 		return !OnIsDebuggerReady.Execute();
 	}
diff --git a/Engine/Source/Editor/BehaviorTreeEditor/Private/SBehaviorTreeBlackboardEditor.h b/Engine/Source/Editor/BehaviorTreeEditor/Private/SBehaviorTreeBlackboardEditor.h
--- a/Engine/Source/Editor/BehaviorTreeEditor/Private/SBehaviorTreeBlackboardEditor.h
+++ b/Engine/Source/Editor/BehaviorTreeEditor/Private/SBehaviorTreeBlackboardEditor.h
@@ -52,6 +52,9 @@ private:
 
 	/** Handle deleting the currently selected entry */
 	void HandleDeleteEntry();
+	
+	/** Handle duplicating the currently selected entry */
+	void HandleDuplicateEntry();
 
 	/** Handle renaming the currently selected entry */
 	void HandleRenameEntry();
@@ -67,6 +70,9 @@ private:
 
 	/** Check whether the 'Delete' operation can be performed on the selected item  */
 	bool CanDeleteEntry() const;
+	
+	/** Check whether the 'Duplicate' operation can be performed on the selected item  */
+	bool CanDuplicateEntry() const;
 
 	/** Check whether the 'Rename' operation can be performed on the selected item  */
 	bool CanRenameEntry() const;
diff --git a/Engine/Source/Editor/BehaviorTreeEditor/Private/SBehaviorTreeBlackboardView.cpp b/Engine/Source/Editor/BehaviorTreeEditor/Private/SBehaviorTreeBlackboardView.cpp
--- a/Engine/Source/Editor/BehaviorTreeEditor/Private/SBehaviorTreeBlackboardView.cpp
+++ b/Engine/Source/Editor/BehaviorTreeEditor/Private/SBehaviorTreeBlackboardView.cpp
@@ -640,6 +640,11 @@ FBlackboardEntry* SBehaviorTreeBlackboardView::GetSelectedEntry(bool& bOutIsInhe
 	return nullptr;
 }
 
+bool SBehaviorTreeBlackboardView::SetSelectedEntry(const FName EntryName) const
+{
+	return GraphActionMenu->SelectItemByName(EntryName);
+}
+
 void SBehaviorTreeBlackboardView::SetObject(UBlackboardData* InBlackboardData)
 {
 	BlackboardData = InBlackboardData;
diff --git a/Engine/Source/Editor/BehaviorTreeEditor/Private/SBehaviorTreeBlackboardView.h b/Engine/Source/Editor/BehaviorTreeEditor/Private/SBehaviorTreeBlackboardView.h
--- a/Engine/Source/Editor/BehaviorTreeEditor/Private/SBehaviorTreeBlackboardView.h
+++ b/Engine/Source/Editor/BehaviorTreeEditor/Private/SBehaviorTreeBlackboardView.h
@@ -97,6 +97,14 @@ public:
 	 * @return A pointer to the currently selected entry (nullptr if a category is selected, or nothing at all)
 	 */
 	FBlackboardEntry* GetSelectedEntry(bool& bOutIsInherited) const;
+	
+	/**
+	 * Selects a blackboard item
+	 * @param	EntryName	The name of blackboard item to select
+	 * @return true if the item was selected
+	 */
+	bool SetSelectedEntry(FName EntryName) const;
+
 
 	/**
 	 * Retrieves the blackboard item index currently selected by the user.
