Copy/Paste blackboard entries

#rnx

diff --git a/Engine/Source/Editor/BehaviorTreeEditor/Private/SBehaviorTreeBlackboardEditor.cpp b/Engine/Source/Editor/BehaviorTreeEditor/Private/SBehaviorTreeBlackboardEditor.cpp
--- a/Engine/Source/Editor/BehaviorTreeEditor/Private/SBehaviorTreeBlackboardEditor.cpp
+++ b/Engine/Source/Editor/BehaviorTreeEditor/Private/SBehaviorTreeBlackboardEditor.cpp
@@ -18,11 +18,11 @@
 #include "Framework/MultiBox/MultiBoxBuilder.h"
 #include "Framework/MultiBox/MultiBoxExtender.h"
 #include "Framework/SlateDelegates.h"
+#include "HAL/PlatformApplicationMisc.h"
 #include "HAL/PlatformCrt.h"
 #include "Internationalization/Internationalization.h"
 #include "Internationalization/Text.h"
 #include "Logging/LogCategory.h"
-#include "Math/UnrealMathSSE.h"
 #include "Misc/AssertionMacros.h"
 #include "Misc/Attribute.h"
 #include "Misc/CString.h"
@@ -31,9 +31,7 @@
 #include "ScopedTransaction.h"
 #include "Styling/AppStyle.h"
 #include "Textures/SlateIcon.h"
-#include "Trace/Detail/Channel.h"
 #include "Types/SlateEnums.h"
-#include "Types/SlateStructs.h"
 #include "UObject/Class.h"
 #include "UObject/NameTypes.h"
 #include "UObject/ObjectMacros.h"
@@ -68,6 +66,18 @@ void SBehaviorTreeBlackboardEditor::Construct(const FArguments& InArgs, TSharedR
 		FCanExecuteAction::CreateSP(this, &SBehaviorTreeBlackboardEditor::CanDeleteEntry)
 		);
 
+	CommandList->MapAction(
+		FGenericCommands::Get().Copy,
+		FExecuteAction::CreateSP(this, &SBehaviorTreeBlackboardEditor::HandleCopyEntry),
+		FCanExecuteAction::CreateSP(this, &SBehaviorTreeBlackboardEditor::CanCopyEntry)
+		);
+
+	CommandList->MapAction(
+		FGenericCommands::Get().Paste,
+		FExecuteAction::CreateSP(this, &SBehaviorTreeBlackboardEditor::HandlePasteEntry),
+		FCanExecuteAction::CreateSP(this, &SBehaviorTreeBlackboardEditor::CanPasteEntry)
+		);
+
 	CommandList->MapAction(
 		FGenericCommands::Get().Duplicate,
 		FExecuteAction::CreateSP(this, &SBehaviorTreeBlackboardEditor::HandleDuplicateEntry),
@@ -101,6 +111,8 @@ void SBehaviorTreeBlackboardEditor::FillContextMenu(FMenuBuilder& MenuBuilder) c
 	if (!IsDebuggerActive() && HasSelectedItems())
 	{
 		MenuBuilder.AddMenuEntry(FBTBlackboardCommands::Get().DeleteEntry);
+		MenuBuilder.AddMenuEntry(FGenericCommands::Get().Copy, NAME_None, LOCTEXT("Copy", "Copy"), LOCTEXT("Copy_Tooltip", "Copies this blackboard entry.") );
+		MenuBuilder.AddMenuEntry(FGenericCommands::Get().Paste, NAME_None, LOCTEXT("Paste", "Paste"), LOCTEXT("Paste_Tooltip", "Pastes this blackboard entry.") );
 		MenuBuilder.AddMenuEntry(FGenericCommands::Get().Duplicate, NAME_None, LOCTEXT("Duplicate", "Duplicate"), LOCTEXT("Duplicate_Tooltip", "Duplicates this blackboard entry.") );
 		MenuBuilder.AddMenuEntry(FGenericCommands::Get().Rename, NAME_None, LOCTEXT("Rename", "Rename"), LOCTEXT("Rename_Tooltip", "Renames this blackboard entry.") );
 	}
@@ -171,6 +183,118 @@ void SBehaviorTreeBlackboardEditor::HandleDeleteEntry()
 	}
 }
 
+namespace BehaviorTreeBlackboardEditorHelper
+{
+	constexpr char PathDelimiter = '/';
+	
+	FName GetUniqueName(const UBlackboardData& BlackboardData, const FName& BaseName)
+	{
+		if (BlackboardData.GetKeyID(BaseName) != FBlackboard::InvalidKey)
+		{
+			FString BaseNameString = BaseName.ToString();
+			FString UniqueName = BaseNameString + "Copy";
+			int32 index = 1;
+			while (BlackboardData.GetKeyID(FName(UniqueName)) != FBlackboard::InvalidKey)
+			{
+				UniqueName = BaseNameString + "Copy(" + FString::FromInt(index++) + ')';
+			}
+			return FName(UniqueName);
+		}
+		return BaseName;
+	}
+	
+	FString GetBlackboardEntryPath(const UBlackboardData& BlackboardData, const FBlackboardEntry& BlackboardEntry)
+	{
+		return BlackboardData.GetPathName() + PathDelimiter + BlackboardEntry.EntryName.ToString();
+	}
+	
+	const FBlackboardEntry* GetBlackboardEntry(const FString& BlackBoardEntryPath)
+	{
+		int32 ItemIndex = 0;
+		if (BlackBoardEntryPath.FindLastChar(PathDelimiter, ItemIndex))
+		{
+			UBlackboardData* SourceBlackboardDataPtr = nullptr;
+			while (true)
+			{
+				FString AssetString = BlackBoardEntryPath.Left(ItemIndex);
+				FSoftObjectPath ObjectPath(AssetString);
+				SourceBlackboardDataPtr = Cast<UBlackboardData>(ObjectPath.TryLoad());
+				if (SourceBlackboardDataPtr != nullptr || !AssetString.FindLastChar(PathDelimiter, ItemIndex))
+				{
+					break;
+				}
+			}
+
+			if (SourceBlackboardDataPtr)
+			{
+				FString ValueString = BlackBoardEntryPath.Right(BlackBoardEntryPath.Len() - ItemIndex - 1);
+				return SourceBlackboardDataPtr->GetKey(SourceBlackboardDataPtr->GetKeyID(FName(ValueString)));
+			}
+		}
+		return nullptr;
+	}
+}
+
+void SBehaviorTreeBlackboardEditor::HandleCopyEntry()
+{
+	if (IsDebuggerActive())
+	{
+		return;
+	}
+
+	UBlackboardData* BlackboardDataPtr = BlackboardData.Get();
+	if (!BlackboardDataPtr)
+	{
+		UE_LOG(LogBlackboardEditor, Error, TEXT("Trying to copy an entry from a blackboard while no Blackboard Asset is set!"));
+		return;
+	}
+
+	bool bIsInherited = false;
+	FBlackboardEntry* BlackboardEntry = GetSelectedEntry(bIsInherited);
+	if (BlackboardEntry != nullptr && !BlackboardEntry->EntryName.IsNone())
+	{
+		FString ClipboardMessage = BehaviorTreeBlackboardEditorHelper::GetBlackboardEntryPath(*BlackboardDataPtr, *BlackboardEntry);
+		FPlatformApplicationMisc::ClipboardCopy(*ClipboardMessage);
+	}
+}
+
+void SBehaviorTreeBlackboardEditor::HandlePasteEntry()
+{
+	UBlackboardData* BlackboardDataPtr = BlackboardData.Get();
+	if (!BlackboardDataPtr)
+	{
+		UE_LOG(LogBlackboardEditor, Error, TEXT("Trying to paste an entry from a blackboard while no Blackboard Asset is set!"));
+		return;
+	}
+
+	FString ClipboardMessage;
+	FPlatformApplicationMisc::ClipboardPaste(ClipboardMessage);
+	const FBlackboardEntry* SourceEntry = BehaviorTreeBlackboardEditorHelper::GetBlackboardEntry(ClipboardMessage);
+	if (SourceEntry)
+	{
+		const FScopedTransaction Transaction(LOCTEXT("BlackboardEntryPasteTransaction", "Paste Blackboard Entry"));
+		BlackboardDataPtr->SetFlags(RF_Transactional);
+		BlackboardDataPtr->Modify();
+
+		FProperty* KeysProperty = FindFProperty<FProperty>(UBlackboardData::StaticClass(), GET_MEMBER_NAME_CHECKED(UBlackboardData, Keys));
+		BlackboardDataPtr->PreEditChange(KeysProperty);
+
+		FBlackboardEntry NewEntry = *SourceEntry;
+		NewEntry.EntryName = BehaviorTreeBlackboardEditorHelper::GetUniqueName(*BlackboardDataPtr, SourceEntry->EntryName);
+		NewEntry.KeyType = DuplicateObject<UBlackboardKeyType>(SourceEntry->KeyType, BlackboardDataPtr);
+		
+		BlackboardDataPtr->Keys.Add(NewEntry);
+
+		GraphActionMenu->RefreshAllActions(true);
+		UE::BehaviorTreeEditor::Delegates::OnBlackboardKeyChanged.Broadcast(*BlackboardDataPtr, &BlackboardDataPtr->Keys.Last());
+
+		FPropertyChangedEvent PropertyChangedEvent(KeysProperty, EPropertyChangeType::ArrayAdd);
+		BlackboardDataPtr->PostEditChangeProperty(PropertyChangedEvent);
+	
+		SetSelectedEntry(NewEntry.EntryName);
+	}
+}
+
 void SBehaviorTreeBlackboardEditor::HandleDuplicateEntry()
 {
 	UBlackboardData* BlackboardDataPtr = BlackboardData.Get();
@@ -194,7 +318,7 @@ void SBehaviorTreeBlackboardEditor::HandleDuplicateEntry()
 			BlackboardDataPtr->PreEditChange(KeysProperty);
 
 			FBlackboardEntry NewEntry = *BlackboardEntry;
-			NewEntry.EntryName = FName(BlackboardEntry->EntryName.ToString() + TEXT("Copy"));
+			NewEntry.EntryName = BehaviorTreeBlackboardEditorHelper::GetUniqueName(*BlackboardDataPtr, BlackboardEntry->EntryName);
 			NewEntry.KeyType = DuplicateObject<UBlackboardKeyType>(BlackboardEntry->KeyType, BlackboardDataPtr);
 			BlackboardDataPtr->Keys.Add(NewEntry);
 
@@ -237,6 +361,38 @@ bool SBehaviorTreeBlackboardEditor::CanDeleteEntry() const
 	return false;
 }
 
+bool SBehaviorTreeBlackboardEditor::CanCopyEntry() const
+{
+	const bool bModeActive = OnIsBlackboardModeActive.IsBound() && OnIsBlackboardModeActive.Execute();
+
+	if (!IsDebuggerActive() && bModeActive)
+	{
+		bool bIsInherited = false;
+		FBlackboardEntry* BlackboardEntry = GetSelectedEntry(bIsInherited);
+		if (BlackboardEntry != nullptr)
+		{
+			return true;
+		}
+	}
+
+	return false;
+}
+
+bool SBehaviorTreeBlackboardEditor::CanPasteEntry() const
+{
+	const bool bModeActive = OnIsBlackboardModeActive.IsBound() && OnIsBlackboardModeActive.Execute();
+
+	if (!IsDebuggerActive() && bModeActive)
+	{
+		FString ClipboardMessage;
+		FPlatformApplicationMisc::ClipboardPaste(ClipboardMessage);
+		const FBlackboardEntry* SourceEntry = BehaviorTreeBlackboardEditorHelper::GetBlackboardEntry(ClipboardMessage);
+		return SourceEntry != nullptr;
+	}
+
+	return false;
+}
+
 bool SBehaviorTreeBlackboardEditor::CanDuplicateEntry() const
 {
 	const bool bModeActive = OnIsBlackboardModeActive.IsBound() && OnIsBlackboardModeActive.Execute();
diff --git a/Engine/Source/Editor/BehaviorTreeEditor/Private/SBehaviorTreeBlackboardEditor.h b/Engine/Source/Editor/BehaviorTreeEditor/Private/SBehaviorTreeBlackboardEditor.h
--- a/Engine/Source/Editor/BehaviorTreeEditor/Private/SBehaviorTreeBlackboardEditor.h
+++ b/Engine/Source/Editor/BehaviorTreeEditor/Private/SBehaviorTreeBlackboardEditor.h
@@ -52,6 +52,12 @@ private:
 
 	/** Handle deleting the currently selected entry */
 	void HandleDeleteEntry();
+
+	/** Handle copying the currently selected entry */
+	void HandleCopyEntry();
+
+	/** Handle pasting the currently copied entry */
+	void HandlePasteEntry();
 	
 	/** Handle duplicating the currently selected entry */
 	void HandleDuplicateEntry();
@@ -70,6 +76,12 @@ private:
 
 	/** Check whether the 'Delete' operation can be performed on the selected item  */
 	bool CanDeleteEntry() const;
+
+	/** Check whether the 'Copy' operation can be performed on the selected item  */
+	bool CanCopyEntry() const;
+
+	/** Check whether the 'Paste' operation can be performed on the copied item  */
+	bool CanPasteEntry() const;
 	
 	/** Check whether the 'Duplicate' operation can be performed on the selected item  */
 	bool CanDuplicateEntry() const;
