Use ImportText and ExportText to copy/paste blackboard entries

#rnx

diff --git a/Engine/Source/Editor/BehaviorTreeEditor/Private/SBehaviorTreeBlackboardEditor.cpp b/Engine/Source/Editor/BehaviorTreeEditor/Private/SBehaviorTreeBlackboardEditor.cpp
--- a/Engine/Source/Editor/BehaviorTreeEditor/Private/SBehaviorTreeBlackboardEditor.cpp
+++ b/Engine/Source/Editor/BehaviorTreeEditor/Private/SBehaviorTreeBlackboardEditor.cpp
@@ -185,7 +185,7 @@ void SBehaviorTreeBlackboardEditor::HandleDeleteEntry()
 
 namespace BehaviorTreeBlackboardEditorHelper
 {
-	constexpr char PathDelimiter = '/';
+	constexpr const char* BlackboardEntryIdentifier = "BlackBoardEntry:";
 	
 	FName GetUniqueName(const UBlackboardData& BlackboardData, const FName& BaseName)
 	{
@@ -202,37 +202,6 @@ namespace BehaviorTreeBlackboardEditorHelper
 		}
 		return BaseName;
 	}
-	
-	FString GetBlackboardEntryPath(const UBlackboardData& BlackboardData, const FBlackboardEntry& BlackboardEntry)
-	{
-		return BlackboardData.GetPathName() + PathDelimiter + BlackboardEntry.EntryName.ToString();
-	}
-	
-	const FBlackboardEntry* GetBlackboardEntry(const FString& BlackBoardEntryPath)
-	{
-		int32 ItemIndex = 0;
-		if (BlackBoardEntryPath.FindLastChar(PathDelimiter, ItemIndex))
-		{
-			UBlackboardData* SourceBlackboardDataPtr = nullptr;
-			while (true)
-			{
-				FString AssetString = BlackBoardEntryPath.Left(ItemIndex);
-				FSoftObjectPath ObjectPath(AssetString);
-				SourceBlackboardDataPtr = Cast<UBlackboardData>(ObjectPath.TryLoad());
-				if (SourceBlackboardDataPtr != nullptr || !AssetString.FindLastChar(PathDelimiter, ItemIndex))
-				{
-					break;
-				}
-			}
-
-			if (SourceBlackboardDataPtr)
-			{
-				FString ValueString = BlackBoardEntryPath.Right(BlackBoardEntryPath.Len() - ItemIndex - 1);
-				return SourceBlackboardDataPtr->GetKey(SourceBlackboardDataPtr->GetKeyID(FName(ValueString)));
-			}
-		}
-		return nullptr;
-	}
 }
 
 void SBehaviorTreeBlackboardEditor::HandleCopyEntry()
@@ -253,8 +222,11 @@ void SBehaviorTreeBlackboardEditor::HandleCopyEntry()
 	FBlackboardEntry* BlackboardEntry = GetSelectedEntry(bIsInherited);
 	if (BlackboardEntry != nullptr && !BlackboardEntry->EntryName.IsNone())
 	{
-		FString ClipboardMessage = BehaviorTreeBlackboardEditorHelper::GetBlackboardEntryPath(*BlackboardDataPtr, *BlackboardEntry);
-		FPlatformApplicationMisc::ClipboardCopy(*ClipboardMessage);
+		FString BlackboardEntryString;
+		FBlackboardEntry DefaultValues;
+		FBlackboardEntry::StaticStruct()->ExportText(BlackboardEntryString, BlackboardEntry, &DefaultValues, nullptr, PPF_SerializedAsImportText, nullptr);
+		
+		FPlatformApplicationMisc::ClipboardCopy(*(BehaviorTreeBlackboardEditorHelper::BlackboardEntryIdentifier + BlackboardEntryString));
 	}
 }
 
@@ -269,30 +241,35 @@ void SBehaviorTreeBlackboardEditor::HandlePasteEntry()
 
 	FString ClipboardMessage;
 	FPlatformApplicationMisc::ClipboardPaste(ClipboardMessage);
-	const FBlackboardEntry* SourceEntry = BehaviorTreeBlackboardEditorHelper::GetBlackboardEntry(ClipboardMessage);
-	if (SourceEntry)
+	if (!ClipboardMessage.RemoveFromStart(BehaviorTreeBlackboardEditorHelper::BlackboardEntryIdentifier))
 	{
-		const FScopedTransaction Transaction(LOCTEXT("BlackboardEntryPasteTransaction", "Paste Blackboard Entry"));
-		BlackboardDataPtr->SetFlags(RF_Transactional);
-		BlackboardDataPtr->Modify();
-
-		FProperty* KeysProperty = FindFProperty<FProperty>(UBlackboardData::StaticClass(), GET_MEMBER_NAME_CHECKED(UBlackboardData, Keys));
-		BlackboardDataPtr->PreEditChange(KeysProperty);
+		return;
+	}
 
-		FBlackboardEntry NewEntry = *SourceEntry;
-		NewEntry.EntryName = BehaviorTreeBlackboardEditorHelper::GetUniqueName(*BlackboardDataPtr, SourceEntry->EntryName);
-		NewEntry.KeyType = DuplicateObject<UBlackboardKeyType>(SourceEntry->KeyType, BlackboardDataPtr);
-		
-		BlackboardDataPtr->Keys.Add(NewEntry);
+	const FScopedTransaction Transaction(LOCTEXT("BlackboardEntryPasteTransaction", "Paste Blackboard Entry"));
+	BlackboardDataPtr->SetFlags(RF_Transactional);
+	BlackboardDataPtr->Modify();
+	
+	FProperty* KeysProperty = FindFProperty<FProperty>(UBlackboardData::StaticClass(), GET_MEMBER_NAME_CHECKED(UBlackboardData, Keys));
+	BlackboardDataPtr->PreEditChange(KeysProperty);
+	
+	FBlackboardEntry& NewEntry = BlackboardDataPtr->Keys.AddDefaulted_GetRef();
 
+	if (FBlackboardEntry::StaticStruct()->ImportText(*ClipboardMessage, &NewEntry, nullptr, PPF_None, nullptr, FString()))
+	{		
+		NewEntry.EntryName = BehaviorTreeBlackboardEditorHelper::GetUniqueName(*BlackboardDataPtr, NewEntry.EntryName);
 		GraphActionMenu->RefreshAllActions(true);
 		UE::BehaviorTreeEditor::Delegates::OnBlackboardKeyChanged.Broadcast(*BlackboardDataPtr, &BlackboardDataPtr->Keys.Last());
 
 		FPropertyChangedEvent PropertyChangedEvent(KeysProperty, EPropertyChangeType::ArrayAdd);
 		BlackboardDataPtr->PostEditChangeProperty(PropertyChangedEvent);
-	
+
 		SetSelectedEntry(NewEntry.EntryName);
 	}
+	else
+	{
+		BlackboardDataPtr->Keys.Pop();
+	}
 }
 
 void SBehaviorTreeBlackboardEditor::HandleDuplicateEntry()
@@ -386,8 +363,7 @@ bool SBehaviorTreeBlackboardEditor::CanPasteEntry() const
 	{
 		FString ClipboardMessage;
 		FPlatformApplicationMisc::ClipboardPaste(ClipboardMessage);
-		const FBlackboardEntry* SourceEntry = BehaviorTreeBlackboardEditorHelper::GetBlackboardEntry(ClipboardMessage);
-		return SourceEntry != nullptr;
+		return ClipboardMessage.StartsWith(BehaviorTreeBlackboardEditorHelper::BlackboardEntryIdentifier);
 	}
 
 	return false;
