Make behavior tree debugger follow rewind debugger scrubbing
Add Open Parent/Open subtree button to behavior tree debugger

diff --git a/Engine/Source/Editor/BehaviorTreeEditor/BehaviorTreeEditor.Build.cs b/Engine/Source/Editor/BehaviorTreeEditor/BehaviorTreeEditor.Build.cs
--- a/Engine/Source/Editor/BehaviorTreeEditor/BehaviorTreeEditor.Build.cs
+++ b/Engine/Source/Editor/BehaviorTreeEditor/BehaviorTreeEditor.Build.cs
@@ -22,9 +22,11 @@ public class BehaviorTreeEditor : ModuleRules
 				"InputCore",
 				"KismetWidgets",
 				"PropertyEditor",
+				"RewindDebuggerInterface",
 				"Slate",
 				"SlateCore",
 				"ToolMenus",
+				"TraceServices",
 				"UnrealEd",
 			}
 		);
diff --git a/Engine/Source/Editor/BehaviorTreeEditor/Private/BehaviorTreeDebugger.cpp b/Engine/Source/Editor/BehaviorTreeEditor/Private/BehaviorTreeDebugger.cpp
--- a/Engine/Source/Editor/BehaviorTreeEditor/Private/BehaviorTreeDebugger.cpp
+++ b/Engine/Source/Editor/BehaviorTreeEditor/Private/BehaviorTreeDebugger.cpp
@@ -1,29 +1,30 @@
 // Copyright Epic Games, Inc. All Rights Reserved.
 
 #include "BehaviorTreeDebugger.h"
-#include "GameFramework/Actor.h"
-#include "GameFramework/Pawn.h"
-#include "GameFramework/Controller.h"
-#include "EngineGlobals.h"
-#include "EdGraph/EdGraph.h"
-#include "Editor/EditorEngine.h"
-#include "Editor.h"
-#include "BehaviorTree/BTTaskNode.h"
+#include "BehaviorTree/BehaviorTree.h"
 #include "BehaviorTree/BTAuxiliaryNode.h"
+#include "BehaviorTree/BTTaskNode.h"
 #include "BehaviorTree/Tasks/BTTask_RunBehaviorDynamic.h"
-#include "BehaviorTreeGraphNode_CompositeDecorator.h"
+#include "BehaviorTreeDelegates.h"
 #include "BehaviorTreeEditor.h"
-#include "Editor/UnrealEdEngine.h"
-#include "Engine/Selection.h"
-#include "GameFramework/PlayerController.h"
-#include "EngineUtils.h"
-#include "UnrealEdGlobals.h"
+#include "BehaviorTreeGraphNode_CompositeDecorator.h"
 #include "BehaviorTreeGraphNode_Decorator.h"
 #include "BehaviorTreeGraphNode_Root.h"
 #include "BehaviorTreeGraphNode_Service.h"
-#include "BehaviorTree/BehaviorTree.h"
-#include "BehaviorTreeDelegates.h"
+#include "EdGraph/EdGraph.h"
+#include "Editor.h"
+#include "Editor/EditorEngine.h"
+#include "Editor/UnrealEdEngine.h"
+#include "Engine/Selection.h"
+#include "EngineGlobals.h"
+#include "EngineUtils.h"
 #include "Framework/Application/SlateApplication.h"
+#include "GameFramework/Actor.h"
+#include "GameFramework/Controller.h"
+#include "GameFramework/Pawn.h"
+#include "GameFramework/PlayerController.h"
+#include "IRewindDebugger.h"
+#include "UnrealEdGlobals.h"
 
 FBehaviorTreeDebugger::FBehaviorTreeDebugger()
 {
@@ -38,6 +39,7 @@ FBehaviorTreeDebugger::FBehaviorTreeDebugger()
 	StepOutIdx = INDEX_NONE;
 	SavedTimestamp = 0.;
 	CurrentTimestamp = 0.;
+	RewindDebuggerTimestamp = 0.;
 
 	FEditorDelegates::BeginPIE.AddRaw(this, &FBehaviorTreeDebugger::OnBeginPIE);
 	FEditorDelegates::EndPIE.AddRaw(this, &FBehaviorTreeDebugger::OnEndPIE);
@@ -135,12 +137,12 @@ void FBehaviorTreeDebugger::Refresh()
 
 void FBehaviorTreeDebugger::Tick(float DeltaTime)
 {
-	if (TreeAsset == NULL || IsPlaySessionPaused())
+	if (TreeAsset == NULL)
 	{
 		return;
 	}
 
-	if (!TreeInstance.IsValid())
+	if (!TreeInstance.IsValid() && !IsPlaySessionPaused())
 	{
 		// clear state when active tree is lost
 		if (DebuggerInstanceIndex != INDEX_NONE)
@@ -152,6 +154,29 @@ void FBehaviorTreeDebugger::Tick(float DeltaTime)
 	}
 
 #if USE_BEHAVIORTREE_DEBUGGER
+	if (IsPlaySessionPaused())
+	{
+		if (const IRewindDebugger* RewindDebugger = IRewindDebugger::Instance())
+		{
+			if (RewindDebugger->IsRecording() && TreeInstance.IsValid())
+			{
+				const double WorldTime = TreeInstance->GetWorld()->GetTimeSeconds();
+				const double RewindDebuggerWorldTime = WorldTime - RewindDebugger->GetRecordingDuration() + RewindDebugger->GetScrubTime();
+				if(RewindDebuggerWorldTime != RewindDebuggerTimestamp)
+				{
+					RewindDebuggerTimestamp = RewindDebuggerWorldTime;
+					const int32 DebugStep = FindDebuggerStepForTimeStamp(RewindDebuggerTimestamp);
+					if(DebugStep != INDEX_NONE && DebugStep != ActiveStepIndex)
+					{
+						const int32 PreviousStep = ActiveStepIndex != INDEX_NONE ? ActiveStepIndex : FMath::Max(DebugStep - 1, 0); // If debug step is 0 it means the previous step got lost to the max number of step recorded and initialize from itself which isn't ideal but best we can do.
+						UpdateCurrentStep(PreviousStep, DebugStep, false); // Don't open asset while scrubbing to avoid stealing focus and opening tree that might not be of interest.
+					}
+				}
+			}
+		}
+		return;
+	}
+
 	TArray<uint16> EmptyPath;
 
 	int32 TestStepIndex = 0;
@@ -266,11 +291,10 @@ bool FBehaviorTreeDebugger::IsTickable() const
 void FBehaviorTreeDebugger::OnBeginPIE(const bool bIsSimulating)
 {
 	bIsPIEActive = true;
-	if(EditorOwner.IsValid())
+	if (TSharedPtr<FBehaviorTreeEditor> Owner = EditorOwner.Pin())
 	{
-		TSharedPtr<FBehaviorTreeEditor> EditorOwnerPtr = EditorOwner.Pin();
-		EditorOwnerPtr->RegenerateMenusAndToolbars();
-		EditorOwnerPtr->DebuggerUpdateGraph();
+		Owner->RegenerateMenusAndToolbars();
+		Owner->DebuggerUpdateGraph();
 	}
 
 	ActiveBreakpoints.Reset();
@@ -291,9 +315,9 @@ void FBehaviorTreeDebugger::OnBeginPIE(const bool bIsSimulating)
 void FBehaviorTreeDebugger::OnEndPIE(const bool bIsSimulating)
 {
 	bIsPIEActive = false;
-	if(EditorOwner.IsValid())
+	if(TSharedPtr<FBehaviorTreeEditor> Owner = EditorOwner.Pin())
 	{
-		EditorOwner.Pin()->RegenerateMenusAndToolbars();
+		Owner->RegenerateMenusAndToolbars();
 	}
 
 	USelection::SelectObjectEvent.RemoveAll(this);
@@ -432,9 +456,9 @@ void FBehaviorTreeDebugger::SetNodeFlags(const struct FBehaviorTreeDebuggerInsta
 	Node->bDebuggerMarkBreakpointTrigger = NodeInstance->GetExecutionIndex() == StoppedOnBreakpointExecutionIndex;
 	if (Node->bDebuggerMarkBreakpointTrigger)
 	{
-		if(EditorOwner.IsValid())
+		if (TSharedPtr<FBehaviorTreeEditor> Owner = EditorOwner.Pin())
 		{
-			EditorOwner.Pin()->JumpToNode(Node);
+			Owner->JumpToNode(Node);
 		}
 	}
 
@@ -733,6 +757,23 @@ int32 FBehaviorTreeDebugger::FindMatchingDebuggerStack(UBehaviorTreeComponent& T
 	return INDEX_NONE;
 }
 
+int32 FBehaviorTreeDebugger::FindDebuggerStepForTimeStamp(double Timestamp) const
+{
+#if USE_BEHAVIORTREE_DEBUGGER
+	if (TreeInstance.IsValid())
+	{
+		for(int32 Idx = TreeInstance->DebuggerSteps.Num() - 1; Idx > 0; --Idx)
+		{
+			if (TreeInstance->DebuggerSteps[Idx].TimeStamp < Timestamp)
+			{
+				return Idx;
+			}
+		}
+	}
+#endif // USE_USE_BEHAVIORTREE_DEBUGGER
+	return INDEX_NONE;
+}
+
 UBehaviorTreeComponent* FBehaviorTreeDebugger::FindInstanceInActor(AActor* TestActor)
 {
 	UBehaviorTreeComponent* FoundInstance = NULL;
@@ -939,11 +980,71 @@ void FBehaviorTreeDebugger::StepOut()
 #endif
 }
 
+void FBehaviorTreeDebugger::OpenParentTree()
+{
+#if USE_BEHAVIORTREE_DEBUGGER
+	if (ActiveStepIndex != INDEX_NONE)
+	{
+		if (TreeInstance.IsValid())
+		{
+			if (TSharedPtr<FBehaviorTreeEditor> Owner = EditorOwner.Pin())
+			{
+				Owner->DebuggerSwitchAsset(TreeInstance->DebuggerSteps[ActiveStepIndex].InstanceStack[DebuggerInstanceIndex - 1].TreeAsset);
+			}
+		}
+	}
+#endif // USE_BEHAVIORTREE_DEBUGGER
+}
+
+void FBehaviorTreeDebugger::OpenSubtree()
+{
+#if USE_BEHAVIORTREE_DEBUGGER
+	if (ActiveStepIndex != INDEX_NONE)
+	{
+		if (TreeInstance.IsValid())
+		{
+			if (TSharedPtr<FBehaviorTreeEditor> Owner = EditorOwner.Pin())
+			{
+				Owner->DebuggerSwitchAsset(TreeInstance->DebuggerSteps[ActiveStepIndex].InstanceStack[DebuggerInstanceIndex + 1].TreeAsset);
+			}
+		}
+	}
+#endif // USE_BEHAVIORTREE_DEBUGGER
+}
+
 bool FBehaviorTreeDebugger::CanStepOut() const
 {
 	return AreAllGameWorldPaused() && (StepOutIdx != INDEX_NONE);
 }
 
+bool FBehaviorTreeDebugger::CanOpenParentTree() const
+{
+#if USE_BEHAVIORTREE_DEBUGGER
+	if (ActiveStepIndex != INDEX_NONE)
+	{
+		if (TreeInstance.IsValid())
+		{
+			return TreeInstance->DebuggerSteps[ActiveStepIndex].InstanceStack.IsValidIndex(DebuggerInstanceIndex - 1);
+		}
+	}
+#endif // USE_BEHAVIORTREE_DEBUGGER
+	return false;
+}
+
+bool FBehaviorTreeDebugger::CanOpenSubtree() const
+{
+#if USE_BEHAVIORTREE_DEBUGGER
+	if (ActiveStepIndex != INDEX_NONE)
+	{
+		if (TreeInstance.IsValid())
+		{
+			return TreeInstance->DebuggerSteps[ActiveStepIndex].InstanceStack.IsValidIndex(DebuggerInstanceIndex + 1);
+		}
+	}
+#endif // USE_BEHAVIORTREE_DEBUGGER
+	return false;
+}
+
 void FBehaviorTreeDebugger::StepBackInto()
 {
 #if USE_BEHAVIORTREE_DEBUGGER
@@ -968,7 +1069,7 @@ bool FBehaviorTreeDebugger::CanStepBackOver() const
 	return AreAllGameWorldPaused() && (StepBackOverIdx != INDEX_NONE);
 }
 
-void FBehaviorTreeDebugger::UpdateCurrentStep(int32 PrevStepIdx, int32 NewStepIdx)
+void FBehaviorTreeDebugger::UpdateCurrentStep(int32 PrevStepIdx, int32 NewStepIdx, bool bSwitchAsset)
 {
 #if USE_BEHAVIORTREE_DEBUGGER
 	if (TreeInstance.IsValid() && TreeInstance->DebuggerSteps.IsValidIndex(NewStepIdx))
@@ -986,9 +1087,12 @@ void FBehaviorTreeDebugger::UpdateCurrentStep(int32 PrevStepIdx, int32 NewStepId
 			if (CurInstaceIdx != NewInstanceIdx || 
 				CurStepInfo.InstanceStack[CurInstaceIdx].TreeAsset != NewStepInfo.InstanceStack[NewInstanceIdx].TreeAsset)
 			{
-				if(EditorOwner.IsValid())
+				if(bSwitchAsset)
 				{
-					EditorOwner.Pin()->DebuggerSwitchAsset(NewStepInfo.InstanceStack[NewInstanceIdx].TreeAsset);
+					if (TSharedPtr<FBehaviorTreeEditor> Owner = EditorOwner.Pin())
+					{
+						Owner->DebuggerSwitchAsset(NewStepInfo.InstanceStack[NewInstanceIdx].TreeAsset);
+					}
 				}
 				UpdateCurrentSubtree();
 			}
@@ -1081,9 +1185,9 @@ void FBehaviorTreeDebugger::OnActiveNodeChanged(const TArray<uint16>& ActivePath
 
 	if (bShouldPause)
 	{
-		if (EditorOwner.IsValid())
+		if (TSharedPtr<FBehaviorTreeEditor> Owner = EditorOwner.Pin())
 		{
-			EditorOwner.Pin()->FocusWindow(TreeAsset);
+			Owner->FocusWindow(TreeAsset);
 		}
 
 		PausePlaySession();
@@ -1175,7 +1279,12 @@ void FBehaviorTreeDebugger::OnBreakpointRemoved(class UBehaviorTreeGraphNode* No
 void FBehaviorTreeDebugger::UpdateDebuggerViewOnInstanceChange()
 {
 #if USE_BEHAVIORTREE_DEBUGGER
-	UBlackboardData* BBAsset = EditorOwner.IsValid() ? EditorOwner.Pin()->GetBlackboardData() : nullptr;
+	UBlackboardData* BBAsset = nullptr;
+	if (TSharedPtr<FBehaviorTreeEditor> Owner = EditorOwner.Pin())
+	{
+		BBAsset = Owner->GetBlackboardData();
+	}
+
 	if (TreeInstance.IsValid() &&
 		TreeInstance->DebuggerSteps.IsValidIndex(ActiveStepIndex) &&
 		TreeInstance->DebuggerSteps[ActiveStepIndex].InstanceStack.IsValidIndex(DebuggerInstanceIndex))
diff --git a/Engine/Source/Editor/BehaviorTreeEditor/Private/BehaviorTreeDebugger.h b/Engine/Source/Editor/BehaviorTreeEditor/Private/BehaviorTreeDebugger.h
--- a/Engine/Source/Editor/BehaviorTreeEditor/Private/BehaviorTreeDebugger.h
+++ b/Engine/Source/Editor/BehaviorTreeEditor/Private/BehaviorTreeDebugger.h
@@ -66,11 +66,15 @@ public:
 	void StepForwardInto();
 	void StepForwardOver();
 	void StepOut();
+	void OpenParentTree();
+	void OpenSubtree();
 	bool CanStepBackInto() const;
 	bool CanStepBackOver() const;
 	bool CanStepForwardInto() const;
 	bool CanStepForwardOver() const;
 	bool CanStepOut() const;
+	bool CanOpenParentTree() const;
+	bool CanOpenSubtree() const;
 
 	static void StopPlaySession();
 	static void PausePlaySession();
@@ -161,6 +165,7 @@ private:
 	/** Debugger timestamps */
 	double SavedTimestamp;
 	double CurrentTimestamp;
+	double RewindDebuggerTimestamp;
 
 	/** set value of DebuggerInstanceIndex variable */
 	void UpdateDebuggerInstance();
@@ -178,6 +183,9 @@ private:
 	/** find index on execution instance stack of matching tree asset */
 	int32 FindMatchingDebuggerStack(UBehaviorTreeComponent& TestInstance) const;
 
+	/** Return StepId of the step at that TimeStamp. */
+	int32 FindDebuggerStepForTimeStamp(double Timestamp) const;
+
 	/** find BT component in given actor */
 	UBehaviorTreeComponent* FindInstanceInActor(AActor* TestActor);
 
@@ -217,7 +225,7 @@ private:
 	void UpdateCurrentSubtree();
 
 	/** updates currently displayed execution step */
-	void UpdateCurrentStep(int32 PrevStepIdx, int32 NewStepIdx);
+	void UpdateCurrentStep(int32 PrevStepIdx, int32 NewStepIdx, bool bSwitchAsset = true);
 
 	/** updates button states */
 	void UpdateAvailableActions();
diff --git a/Engine/Source/Editor/BehaviorTreeEditor/Private/BehaviorTreeEditor.cpp b/Engine/Source/Editor/BehaviorTreeEditor/Private/BehaviorTreeEditor.cpp
--- a/Engine/Source/Editor/BehaviorTreeEditor/Private/BehaviorTreeEditor.cpp
+++ b/Engine/Source/Editor/BehaviorTreeEditor/Private/BehaviorTreeEditor.cpp
@@ -1087,6 +1087,16 @@ void FBehaviorTreeEditor::BindDebuggerToolbarCommands()
 		FExecuteAction::CreateSP(DebuggerOb, &FBehaviorTreeDebugger::StepOut),
 		FCanExecuteAction::CreateSP(DebuggerOb, &FBehaviorTreeDebugger::CanStepOut));
 
+	ToolkitCommands->MapAction(
+		Commands.OpenParentTree,
+		FExecuteAction::CreateSP(DebuggerOb, &FBehaviorTreeDebugger::OpenParentTree),
+		FCanExecuteAction::CreateSP(DebuggerOb, &FBehaviorTreeDebugger::CanOpenParentTree));
+
+	ToolkitCommands->MapAction(
+		Commands.OpenSubtree,
+		FExecuteAction::CreateSP(DebuggerOb, &FBehaviorTreeDebugger::OpenSubtree),
+		FCanExecuteAction::CreateSP(DebuggerOb, &FBehaviorTreeDebugger::CanOpenSubtree));
+
 	ToolkitCommands->MapAction(
 		Commands.PausePlaySession,
 		FExecuteAction::CreateStatic(&FBehaviorTreeDebugger::PausePlaySession),
diff --git a/Engine/Source/Editor/BehaviorTreeEditor/Private/BehaviorTreeEditorCommands.cpp b/Engine/Source/Editor/BehaviorTreeEditor/Private/BehaviorTreeEditorCommands.cpp
--- a/Engine/Source/Editor/BehaviorTreeEditor/Private/BehaviorTreeEditorCommands.cpp
+++ b/Engine/Source/Editor/BehaviorTreeEditor/Private/BehaviorTreeEditorCommands.cpp
@@ -36,6 +36,8 @@ void FBTDebuggerCommands::RegisterCommands()
 	UI_COMMAND(ForwardInto, "Forward: Into", "Show state from next step, can go into subtrees", EUserInterfaceActionType::Button, FInputChord());
 	UI_COMMAND(ForwardOver, "Forward: Over", "Show state from next step, don't go into subtrees", EUserInterfaceActionType::Button, FInputChord());
 	UI_COMMAND(StepOut, "Step Out", "Show state from next step, leave current subtree", EUserInterfaceActionType::Button, FInputChord());
+	UI_COMMAND(OpenParentTree, "Open Parent Tree", "Open the parent tree of this subtree", EUserInterfaceActionType::Button, FInputChord());
+	UI_COMMAND(OpenSubtree, "Open Subtree", "Open the currently running subtree", EUserInterfaceActionType::Button, FInputChord());
 	
 	UI_COMMAND(PausePlaySession, "Pause", "Pause simulation", EUserInterfaceActionType::Button, FInputChord());
 	UI_COMMAND(ResumePlaySession, "Resume", "Resume simulation", EUserInterfaceActionType::Button, FInputChord() );
diff --git a/Engine/Source/Editor/BehaviorTreeEditor/Private/BehaviorTreeEditorCommands.h b/Engine/Source/Editor/BehaviorTreeEditor/Private/BehaviorTreeEditorCommands.h
--- a/Engine/Source/Editor/BehaviorTreeEditor/Private/BehaviorTreeEditorCommands.h
+++ b/Engine/Source/Editor/BehaviorTreeEditor/Private/BehaviorTreeEditorCommands.h
@@ -29,6 +29,8 @@ public:
 	TSharedPtr<FUICommandInfo> ForwardInto;
 	TSharedPtr<FUICommandInfo> ForwardOver;
 	TSharedPtr<FUICommandInfo> StepOut;
+	TSharedPtr<FUICommandInfo> OpenParentTree;
+	TSharedPtr<FUICommandInfo> OpenSubtree;
 
 	TSharedPtr<FUICommandInfo> PausePlaySession;
 	TSharedPtr<FUICommandInfo> ResumePlaySession;
diff --git a/Engine/Source/Editor/BehaviorTreeEditor/Private/BehaviorTreeEditorToolbar.cpp b/Engine/Source/Editor/BehaviorTreeEditor/Private/BehaviorTreeEditorToolbar.cpp
--- a/Engine/Source/Editor/BehaviorTreeEditor/Private/BehaviorTreeEditorToolbar.cpp
+++ b/Engine/Source/Editor/BehaviorTreeEditor/Private/BehaviorTreeEditorToolbar.cpp
@@ -95,6 +95,8 @@ void FBehaviorTreeEditorToolbar::AddDebuggerToolbar(TSharedPtr<FExtender> Extend
 					ToolbarBuilder.AddToolBarButton(FBTDebuggerCommands::Get().ForwardInto);
 					ToolbarBuilder.AddToolBarButton(FBTDebuggerCommands::Get().ForwardOver);
 					ToolbarBuilder.AddToolBarButton(FBTDebuggerCommands::Get().StepOut);
+					ToolbarBuilder.AddToolBarButton(FBTDebuggerCommands::Get().OpenParentTree);
+					ToolbarBuilder.AddToolBarButton(FBTDebuggerCommands::Get().OpenSubtree);
 				}
 				ToolbarBuilder.EndSection();
 				ToolbarBuilder.BeginSection("World");
diff --git a/Engine/Source/Editor/EditorStyle/Private/SlateEditorStyle.cpp b/Engine/Source/Editor/EditorStyle/Private/SlateEditorStyle.cpp
--- a/Engine/Source/Editor/EditorStyle/Private/SlateEditorStyle.cpp
+++ b/Engine/Source/Editor/EditorStyle/Private/SlateEditorStyle.cpp
@@ -4110,6 +4110,10 @@ void FSlateEditorStyle::FStyle::SetupGraphEditorStyles()
 			Set( "BTEditor.Debugger.StepOut.Small", new IMAGE_BRUSH("Icons/icon_step_40x", Icon20x20));
 			Set( "BTEditor.Debugger.SingleStep", new IMAGE_BRUSH("Icons/icon_advance_40x", Icon40x40));
 			Set( "BTEditor.Debugger.SingleStep.Small", new IMAGE_BRUSH( "Icons/icon_advance_40x", Icon20x20 ) );
+			Set( "BTEditor.Debugger.OpenParentTree", new IMAGE_BRUSH("Icons/icon_DebugStepOut_40x", Icon40x40));
+			Set( "BTEditor.Debugger.OpenParentTree.Small", new IMAGE_BRUSH("Icons/icon_DebugStepOut_40x", Icon20x20));
+			Set( "BTEditor.Debugger.OpenSubtree", new IMAGE_BRUSH("Icons/icon_DebugStepIn_40x", Icon40x40));
+			Set( "BTEditor.Debugger.OpenSubtree.Small", new IMAGE_BRUSH("Icons/icon_DebugStepIn_40x", Icon20x20));
 
 			Set( "BTEditor.Debugger.PausePlaySession", new IMAGE_BRUSH( "Icons/icon_pause_40x", Icon40x40 ) );
 			Set( "BTEditor.Debugger.PausePlaySession.Small", new IMAGE_BRUSH( "Icons/icon_pause_40x", Icon20x20 ) );
diff --git a/Engine/Source/Editor/EditorStyle/Private/StarshipStyle.cpp b/Engine/Source/Editor/EditorStyle/Private/StarshipStyle.cpp
--- a/Engine/Source/Editor/EditorStyle/Private/StarshipStyle.cpp
+++ b/Engine/Source/Editor/EditorStyle/Private/StarshipStyle.cpp
@@ -3897,6 +3897,10 @@ void FStarshipEditorStyle::FStyle::SetupGraphEditorStyles()
 			Set( "BTEditor.Debugger.StepOut.Small", new IMAGE_BRUSH("Icons/icon_step_40x", Icon20x20));
 			Set( "BTEditor.Debugger.SingleStep", new IMAGE_BRUSH("Icons/icon_advance_40x", Icon40x40));
 			Set( "BTEditor.Debugger.SingleStep.Small", new IMAGE_BRUSH( "Icons/icon_advance_40x", Icon20x20 ) );
+			Set( "BTEditor.Debugger.OpenParentTree", new IMAGE_BRUSH("Icons/icon_DebugStepOut_40x", Icon40x40));
+			Set( "BTEditor.Debugger.OpenParentTree.Small", new IMAGE_BRUSH("Icons/icon_DebugStepOut_40x", Icon20x20));
+			Set( "BTEditor.Debugger.OpenSubtree", new IMAGE_BRUSH("Icons/icon_DebugStepIn_40x", Icon40x40));
+			Set( "BTEditor.Debugger.OpenSubtree.Small", new IMAGE_BRUSH("Icons/icon_DebugStepIn_40x", Icon20x20));
 
 			Set( "BTEditor.Debugger.PausePlaySession", new IMAGE_BRUSH( "Icons/icon_pause_40x", Icon40x40 ) );
 			Set( "BTEditor.Debugger.PausePlaySession.Small", new IMAGE_BRUSH( "Icons/icon_pause_40x", Icon20x20 ) );
