PR #12608: Improve BehaviorTreeDebugger Performance on begin pie.
#rnx

diff --git a/Engine/Source/Editor/BehaviorTreeEditor/Private/BehaviorTreeDebugger.cpp b/Engine/Source/Editor/BehaviorTreeEditor/Private/BehaviorTreeDebugger.cpp
--- a/Engine/Source/Editor/BehaviorTreeEditor/Private/BehaviorTreeDebugger.cpp
+++ b/Engine/Source/Editor/BehaviorTreeEditor/Private/BehaviorTreeDebugger.cpp
@@ -2,6 +2,7 @@
 
 #include "BehaviorTreeDebugger.h"
 #include "BehaviorTree/BehaviorTree.h"
+#include "BehaviorTree/BehaviorTreeManager.h"
 #include "BehaviorTree/BTAuxiliaryNode.h"
 #include "BehaviorTree/BTTaskNode.h"
 #include "BehaviorTree/Tasks/BTTask_RunBehaviorDynamic.h"
@@ -851,24 +852,28 @@ void FBehaviorTreeDebugger::FindMatchingTreeInstance()
 		return;
 	}
 
-	UBehaviorTreeComponent* MatchingComp = NULL;
-	for (FActorIterator It(PlayWorld); It; ++It)
+	UBehaviorTreeManager* BTManager = UBehaviorTreeManager::GetCurrent(PlayWorld);
+	if (!BTManager)
 	{
-		AActor* TestActor = *It;
-		UBehaviorTreeComponent* TestComp = TestActor ? TestActor->FindComponentByClass<UBehaviorTreeComponent>() : nullptr;
+		return;
+	}
 
-		if (TestComp)
+	UBehaviorTreeComponent* MatchingComp = NULL;
+	for (UBehaviorTreeComponent* ActiveBTComp : BTManager->GetActiveComponents())
+	{
+		AActor* TestActor = ActiveBTComp ? ActiveBTComp->GetOwner() : nullptr;
+		if (ActiveBTComp && TestActor)
 		{
-			KnownInstances.Add(TestComp);
+			KnownInstances.Add(ActiveBTComp);
 
-			const int32 MatchingIdx = FindMatchingDebuggerStack(*TestComp);
+			const int32 MatchingIdx = FindMatchingDebuggerStack(*ActiveBTComp);
 			if (MatchingIdx != INDEX_NONE)
 			{
-				MatchingComp = TestComp;
+				MatchingComp = ActiveBTComp;
 
 				if (TestActor->IsSelected())
 				{
-					TreeInstance = TestComp;
+					TreeInstance = ActiveBTComp;
 					return;
 				}
 			}
diff --git a/Engine/Source/Runtime/AIModule/Classes/BehaviorTree/BehaviorTreeManager.h b/Engine/Source/Runtime/AIModule/Classes/BehaviorTree/BehaviorTreeManager.h
--- a/Engine/Source/Runtime/AIModule/Classes/BehaviorTree/BehaviorTreeManager.h
+++ b/Engine/Source/Runtime/AIModule/Classes/BehaviorTree/BehaviorTreeManager.h
@@ -61,6 +61,11 @@ class UBehaviorTreeManager : public UObject
 	static AIMODULE_API UBehaviorTreeManager* GetCurrent(UWorld* World);
 	static AIMODULE_API UBehaviorTreeManager* GetCurrent(UObject* WorldContextObject);
 
+	TConstArrayView<TObjectPtr<UBehaviorTreeComponent>> GetActiveComponents() const 
+	{ 
+	    return ActiveComponents;
+	}
+
 protected:
 
 	/** initialized tree templates */
