Fix copy pasting blackboard key storing a reference to the other Key instanced object.
#rnx

diff --git a/Engine/Source/Editor/BehaviorTreeEditor/Private/SBehaviorTreeBlackboardEditor.cpp b/Engine/Source/Editor/BehaviorTreeEditor/Private/SBehaviorTreeBlackboardEditor.cpp
--- a/Engine/Source/Editor/BehaviorTreeEditor/Private/SBehaviorTreeBlackboardEditor.cpp
+++ b/Engine/Source/Editor/BehaviorTreeEditor/Private/SBehaviorTreeBlackboardEditor.cpp
@@ -255,7 +255,7 @@ void SBehaviorTreeBlackboardEditor::HandlePasteEntry()
 	
 	FBlackboardEntry& NewEntry = BlackboardDataPtr->Keys.AddDefaulted_GetRef();
 
-	if (FBlackboardEntry::StaticStruct()->ImportText(*ClipboardMessage, &NewEntry, nullptr, PPF_None, nullptr, FString()))
+	if (FBlackboardEntry::StaticStruct()->ImportText(*ClipboardMessage, &NewEntry, BlackboardDataPtr, PPF_InstanceSubobjects, nullptr, FString()))
 	{		
 		NewEntry.EntryName = BehaviorTreeBlackboardEditorHelper::GetUniqueName(*BlackboardDataPtr, NewEntry.EntryName);
 		GraphActionMenu->RefreshAllActions(true);
