Updated headers using UnrealCodeFixup to make sure dllstorage is on methods/staticvars instead of types. Part 4/n

diff --git a/Engine/Source/Editor/AudioEditor/Public/SoundCueEditorUtilities.h b/Engine/Source/Editor/AudioEditor/Public/SoundCueEditorUtilities.h
--- a/Engine/Source/Editor/AudioEditor/Public/SoundCueEditorUtilities.h
+++ b/Engine/Source/Editor/AudioEditor/Public/SoundCueEditorUtilities.h
@@ -9,6 +9,8 @@
 #include "Math/Vector2D.h"
 #include "Templates/SharedPointer.h"
 
+#define UE_API AUDIOEDITOR_API
+
 class ISoundCueEditor;
 class UEdGraph;
 class UObject;
@@ -18,15 +20,15 @@ struct Rect;
 //////////////////////////////////////////////////////////////////////////
 // FSoundCueEditorUtilities
 
-class AUDIOEDITOR_API FSoundCueEditorUtilities
+class FSoundCueEditorUtilities
 {
 public:
 
 	/** Can we paste to this graph? */
-	static bool CanPasteNodes(const class UEdGraph* Graph);
+	static UE_API bool CanPasteNodes(const class UEdGraph* Graph);
 
 	/** Perform paste on graph, at location */
-	static void  PasteNodesHere(class UEdGraph* Graph, const FVector2D& Location);
+	static UE_API void  PasteNodesHere(class UEdGraph* Graph, const FVector2D& Location);
 
 	/** Adds SoundNodeWavePlayer nodes based on selected objects
 	 *
@@ -35,7 +37,7 @@ public:
 	 * @param OutPlayers Stores all created nodes
 	 * @param Location Position of first created node
 	 */
-	static void CreateWaveContainers(TArray<class USoundWave*>& SelectedWaves, class USoundCue* SoundCue, TArray<class USoundNode*>& OutPlayers, FVector2D Location);
+	static UE_API void CreateWaveContainers(TArray<class USoundWave*>& SelectedWaves, class USoundCue* SoundCue, TArray<class USoundNode*>& OutPlayers, FVector2D Location);
 
 	/** Adds USoundNodeDialoguePlayer nodes based on selected objects
 	 *
@@ -44,7 +46,7 @@ public:
 	 * @param OutPlayers Stores all created nodes
 	 * @param Location Position of first created node
 	 */
-	static void CreateDialogueContainers(TArray<class UDialogueWave*>& SelectedDialogues, USoundCue* SoundCue, TArray<class USoundNode*>& OutPlayers, FVector2D Location);
+	static UE_API void CreateDialogueContainers(TArray<class UDialogueWave*>& SelectedDialogues, USoundCue* SoundCue, TArray<class USoundNode*>& OutPlayers, FVector2D Location);
 
 
 	/** Get the bounding area for the currently selected nodes
@@ -54,17 +56,19 @@ public:
 	 * @param Padding An amount of padding to add to all sides of the bounds
 	 *
 	 * @return false if nothing is selected*/
-	static bool GetBoundsForSelectedNodes(const UEdGraph* Graph, class FSlateRect& Rect, float Padding = 0.0f);
+	static UE_API bool GetBoundsForSelectedNodes(const UEdGraph* Graph, class FSlateRect& Rect, float Padding = 0.0f);
 
 	/** Gets the number of nodes that are currently selected */
-	static int32 GetNumberOfSelectedNodes(const UEdGraph* Graph);
+	static UE_API int32 GetNumberOfSelectedNodes(const UEdGraph* Graph);
 
 	/** Get the currently selected set of nodes */
-	static TSet<UObject*> GetSelectedNodes(const UEdGraph* Graph);
+	static UE_API TSet<UObject*> GetSelectedNodes(const UEdGraph* Graph);
 
 private:
 	/** Get ISoundCueEditor for given object, if it exists */
-	static TSharedPtr<class ISoundCueEditor> GetISoundCueEditorForObject(const UObject* ObjectToFocusOn);
+	static UE_API TSharedPtr<class ISoundCueEditor> GetISoundCueEditorForObject(const UObject* ObjectToFocusOn);
 
 	FSoundCueEditorUtilities() {}
 };
+
+#undef UE_API
diff --git a/Engine/Source/Editor/AudioEditor/Public/SoundSubmixEditorUtilities.h b/Engine/Source/Editor/AudioEditor/Public/SoundSubmixEditorUtilities.h
--- a/Engine/Source/Editor/AudioEditor/Public/SoundSubmixEditorUtilities.h
+++ b/Engine/Source/Editor/AudioEditor/Public/SoundSubmixEditorUtilities.h
@@ -6,10 +6,12 @@
 #include "CoreMinimal.h"
 #include "Math/Vector2D.h"
 
+#define UE_API AUDIOEDITOR_API
+
 class UEdGraph;
 class UEdGraphPin;
 
-class AUDIOEDITOR_API FSoundSubmixEditorUtilities
+class FSoundSubmixEditorUtilities
 {
 public:
 
@@ -20,8 +22,10 @@ public:
 	 * @param	FromPin		Pin that was dragged to create sound class
 	 * @param	Location	Location for new sound class
 	 */
-	static void CreateSoundSubmix(const UEdGraph* Graph, UEdGraphPin* FromPin, FVector2D Location, const FString& Name);
+	static UE_API void CreateSoundSubmix(const UEdGraph* Graph, UEdGraphPin* FromPin, FVector2D Location, const FString& Name);
 
 private:
 	FSoundSubmixEditorUtilities() {}
 };
+
+#undef UE_API
diff --git a/Engine/Source/Editor/AudioEditor/Public/SubmixDetailsCustomization.h b/Engine/Source/Editor/AudioEditor/Public/SubmixDetailsCustomization.h
--- a/Engine/Source/Editor/AudioEditor/Public/SubmixDetailsCustomization.h
+++ b/Engine/Source/Editor/AudioEditor/Public/SubmixDetailsCustomization.h
@@ -16,6 +16,8 @@
 #include "Types/SlateEnums.h"
 #include "UObject/NameTypes.h"
 
+#define UE_API AUDIOEDITOR_API
+
 class IDetailGroup;
 class IDetailLayoutBuilder;
 class SWidget;
@@ -58,53 +60,55 @@ private:
 	TWeakPtr<FNameSelectorGenerator> WeakThis;
 };
 
-class AUDIOEDITOR_API FSoundfieldSubmixDetailsCustomization : public IDetailCustomization
+class FSoundfieldSubmixDetailsCustomization : public IDetailCustomization
 {
 protected:
 	struct FPrivateToken { explicit FPrivateToken() = default; };
 
 public:
-	FSoundfieldSubmixDetailsCustomization(FPrivateToken);
+	UE_API FSoundfieldSubmixDetailsCustomization(FPrivateToken);
 
 	// Makes a new instance of this detail layout class
-	static TSharedRef<IDetailCustomization> MakeInstance();
+	static UE_API TSharedRef<IDetailCustomization> MakeInstance();
 
 	// IDetailCustomization interface
-	virtual void CustomizeDetails(IDetailLayoutBuilder& DetailLayout) override;
+	UE_API virtual void CustomizeDetails(IDetailLayoutBuilder& DetailLayout) override;
 	// End of IDetailCustomization interface
 
 private:
 	TSharedPtr<FNameSelectorGenerator> SoundfieldFormatNameSelectorGenerator;
 };
 
-class AUDIOEDITOR_API FEndpointSubmixDetailsCustomization : public IDetailCustomization, public FNameSelectorGenerator
+class FEndpointSubmixDetailsCustomization : public IDetailCustomization, public FNameSelectorGenerator
 {
 public:
-	FEndpointSubmixDetailsCustomization(FProtectedToken);
+	UE_API FEndpointSubmixDetailsCustomization(FProtectedToken);
 
 	// Makes a new instance of this detail layout class
-	static TSharedRef<IDetailCustomization> MakeInstance();
+	static UE_API TSharedRef<IDetailCustomization> MakeInstance();
 
 	// IDetailCustomization interface
-	virtual void CustomizeDetails(IDetailLayoutBuilder& DetailLayout) override;
+	UE_API virtual void CustomizeDetails(IDetailLayoutBuilder& DetailLayout) override;
 	// End of IDetailCustomization interface
 
 private:
 	TSharedPtr<FNameSelectorGenerator> EndpointTypeNameSelectorGenerator;
 };
 
-class AUDIOEDITOR_API FSoundfieldEndpointSubmixDetailsCustomization : public IDetailCustomization, public FNameSelectorGenerator
+class FSoundfieldEndpointSubmixDetailsCustomization : public IDetailCustomization, public FNameSelectorGenerator
 {
 public:
-	FSoundfieldEndpointSubmixDetailsCustomization(FProtectedToken);
+	UE_API FSoundfieldEndpointSubmixDetailsCustomization(FProtectedToken);
 
 	// Makes a new instance of this detail layout class
-	static TSharedRef<IDetailCustomization> MakeInstance();
+	static UE_API TSharedRef<IDetailCustomization> MakeInstance();
 
 	// IDetailCustomization interface
-	virtual void CustomizeDetails(IDetailLayoutBuilder& DetailLayout) override;
+	UE_API virtual void CustomizeDetails(IDetailLayoutBuilder& DetailLayout) override;
 	// End of IDetailCustomization interface
 
 private:
 	TSharedPtr<FNameSelectorGenerator> EndpointTypeNameSelectorGenerator;
 };
+
+#undef UE_API
diff --git a/Engine/Source/Editor/BehaviorTreeEditor/Classes/BehaviorTreeDecoratorGraph.h b/Engine/Source/Editor/BehaviorTreeEditor/Classes/BehaviorTreeDecoratorGraph.h
--- a/Engine/Source/Editor/BehaviorTreeEditor/Classes/BehaviorTreeDecoratorGraph.h
+++ b/Engine/Source/Editor/BehaviorTreeEditor/Classes/BehaviorTreeDecoratorGraph.h
@@ -12,25 +12,29 @@
 
 #include "BehaviorTreeDecoratorGraph.generated.h"
 
+#define UE_API BEHAVIORTREEEDITOR_API
+
 class UBehaviorTreeDecoratorGraphNode;
 class UEdGraphPin;
 class UObject;
 
-UCLASS()
-class BEHAVIORTREEEDITOR_API UBehaviorTreeDecoratorGraph : public UEdGraph
+UCLASS(MinimalAPI)
+class UBehaviorTreeDecoratorGraph : public UEdGraph
 {
 	GENERATED_UCLASS_BODY()
 
-	void CollectDecoratorData(TArray<class UBTDecorator*>& DecoratorInstances, TArray<struct FBTDecoratorLogic>& DecoratorOperations) const;
-	int32 SpawnMissingNodes(const TArray<class UBTDecorator*>& NodeInstances, const TArray<struct FBTDecoratorLogic>& Operations, int32 StartIndex);
+	UE_API void CollectDecoratorData(TArray<class UBTDecorator*>& DecoratorInstances, TArray<struct FBTDecoratorLogic>& DecoratorOperations) const;
+	UE_API int32 SpawnMissingNodes(const TArray<class UBTDecorator*>& NodeInstances, const TArray<struct FBTDecoratorLogic>& Operations, int32 StartIndex);
 
 protected:
 
-	const class UBehaviorTreeDecoratorGraphNode* FindRootNode() const;
-	void CollectDecoratorDataWorker(const class UBehaviorTreeDecoratorGraphNode* Node, TArray<class UBTDecorator*>& DecoratorInstances, TArray<struct FBTDecoratorLogic>& DecoratorOperations) const;
+	UE_API const class UBehaviorTreeDecoratorGraphNode* FindRootNode() const;
+	UE_API void CollectDecoratorDataWorker(const class UBehaviorTreeDecoratorGraphNode* Node, TArray<class UBTDecorator*>& DecoratorInstances, TArray<struct FBTDecoratorLogic>& DecoratorOperations) const;
 
-	UEdGraphPin* FindFreePin(UEdGraphNode* Node, EEdGraphPinDirection Direction);
-	UBehaviorTreeDecoratorGraphNode* SpawnMissingNodeWorker(const TArray<class UBTDecorator*>& NodeInstances, const TArray<struct FBTDecoratorLogic>& Operations, int32& Index,
+	UE_API UEdGraphPin* FindFreePin(UEdGraphNode* Node, EEdGraphPinDirection Direction);
+	UE_API UBehaviorTreeDecoratorGraphNode* SpawnMissingNodeWorker(const TArray<class UBTDecorator*>& NodeInstances, const TArray<struct FBTDecoratorLogic>& Operations, int32& Index,
 		UEdGraphNode* ParentGraphNode, int32 ChildIdx);
 
 };
+
+#undef UE_API
diff --git a/Engine/Source/Editor/BehaviorTreeEditor/Classes/BehaviorTreeFactory.h b/Engine/Source/Editor/BehaviorTreeEditor/Classes/BehaviorTreeFactory.h
--- a/Engine/Source/Editor/BehaviorTreeEditor/Classes/BehaviorTreeFactory.h
+++ b/Engine/Source/Editor/BehaviorTreeEditor/Classes/BehaviorTreeFactory.h
@@ -10,20 +10,23 @@
 
 #include "BehaviorTreeFactory.generated.h"
 
+#define UE_API BEHAVIORTREEEDITOR_API
+
 class FFeedbackContext;
 class UClass;
 class UObject;
 
-UCLASS()
-class BEHAVIORTREEEDITOR_API UBehaviorTreeFactory : public UFactory
+UCLASS(MinimalAPI)
+class UBehaviorTreeFactory : public UFactory
 {
 	GENERATED_UCLASS_BODY()
 
 	// UFactory interface
-	virtual UObject* FactoryCreateNew(UClass* Class,UObject* InParent,FName Name,EObjectFlags Flags,UObject* Context,FFeedbackContext* Warn) override;
-	virtual bool CanCreateNew() const override;
+	UE_API virtual UObject* FactoryCreateNew(UClass* Class,UObject* InParent,FName Name,EObjectFlags Flags,UObject* Context,FFeedbackContext* Warn) override;
+	UE_API virtual bool CanCreateNew() const override;
 	// End of UFactory interface
 };
 
 
 
+#undef UE_API
diff --git a/Engine/Source/Editor/BehaviorTreeEditor/Classes/BehaviorTreeGraph.h b/Engine/Source/Editor/BehaviorTreeEditor/Classes/BehaviorTreeGraph.h
--- a/Engine/Source/Editor/BehaviorTreeEditor/Classes/BehaviorTreeGraph.h
+++ b/Engine/Source/Editor/BehaviorTreeEditor/Classes/BehaviorTreeGraph.h
@@ -11,12 +11,14 @@
 
 #include "BehaviorTreeGraph.generated.h"
 
+#define UE_API BEHAVIORTREEEDITOR_API
+
 class UBehaviorTreeGraphNode_Root;
 class UEdGraphNode;
 class UObject;
 
-UCLASS()
-class BEHAVIORTREEEDITOR_API UBehaviorTreeGraph : public UAIGraph
+UCLASS(MinimalAPI)
+class UBehaviorTreeGraph : public UAIGraph
 {
 	GENERATED_UCLASS_BODY()
 
@@ -37,43 +39,45 @@ class BEHAVIORTREEEDITOR_API UBehaviorTreeGraph : public UAIGraph
 	UPROPERTY()
 	TSubclassOf<UBehaviorTreeGraphNode_Root> RootNodeClass;
 
-	virtual void OnCreated() override;
-	virtual void OnLoaded() override;
-	virtual void Initialize() override;
-	void OnSave();
+	UE_API virtual void OnCreated() override;
+	UE_API virtual void OnLoaded() override;
+	UE_API virtual void Initialize() override;
+	UE_API void OnSave();
 
-	virtual void UpdateVersion() override;
-	virtual void MarkVersion() override;
-	virtual void UpdateAsset(int32 UpdateFlags = 0) override;
-	virtual void OnSubNodeDropped() override;
+	UE_API virtual void UpdateVersion() override;
+	UE_API virtual void MarkVersion() override;
+	UE_API virtual void UpdateAsset(int32 UpdateFlags = 0) override;
+	UE_API virtual void OnSubNodeDropped() override;
 
 	virtual bool DoesSupportServices() const { return true; }
 
-	void UpdateBlackboardChange();
-	void UpdateAbortHighlight(struct FAbortDrawHelper& Mode0, struct FAbortDrawHelper& Mode1);
-	void CreateBTFromGraph(class UBehaviorTreeGraphNode* RootEdNode);
-	void SpawnMissingNodes();
-	void UpdatePinConnectionTypes();
-	bool UpdateInjectedNodes();
-	void UpdateBrokenComposites();
-	class UEdGraphNode* FindInjectedNode(int32 Index);
-	void ReplaceNodeConnections(UEdGraphNode* OldNode, UEdGraphNode* NewNode);
-	void RebuildExecutionOrder();
-	void RebuildChildOrder(UEdGraphNode* ParentNode);
-	void SpawnMissingNodesForParallel();
-	void RemoveUnknownSubNodes();
-
-	void AutoArrange();
+	UE_API void UpdateBlackboardChange();
+	UE_API void UpdateAbortHighlight(struct FAbortDrawHelper& Mode0, struct FAbortDrawHelper& Mode1);
+	UE_API void CreateBTFromGraph(class UBehaviorTreeGraphNode* RootEdNode);
+	UE_API void SpawnMissingNodes();
+	UE_API void UpdatePinConnectionTypes();
+	UE_API bool UpdateInjectedNodes();
+	UE_API void UpdateBrokenComposites();
+	UE_API class UEdGraphNode* FindInjectedNode(int32 Index);
+	UE_API void ReplaceNodeConnections(UEdGraphNode* OldNode, UEdGraphNode* NewNode);
+	UE_API void RebuildExecutionOrder();
+	UE_API void RebuildChildOrder(UEdGraphNode* ParentNode);
+	UE_API void SpawnMissingNodesForParallel();
+	UE_API void RemoveUnknownSubNodes();
+
+	UE_API void AutoArrange();
 
 protected:
 
-	void CollectAllNodeInstances(TSet<UObject*>& NodeInstances) override;
+	UE_API void CollectAllNodeInstances(TSet<UObject*>& NodeInstances) override;
 
 #if WITH_EDITOR
-	virtual void PostEditUndo() override;
+	UE_API virtual void PostEditUndo() override;
 #endif
 
-	void UpdateVersion_UnifiedSubNodes();
-	void UpdateVersion_InnerGraphWhitespace();
-	void UpdateVersion_RunBehaviorInSeparateGraph();
+	UE_API void UpdateVersion_UnifiedSubNodes();
+	UE_API void UpdateVersion_InnerGraphWhitespace();
+	UE_API void UpdateVersion_RunBehaviorInSeparateGraph();
 };
+
+#undef UE_API
diff --git a/Engine/Source/Editor/BehaviorTreeEditor/Classes/BehaviorTreeGraphNode.h b/Engine/Source/Editor/BehaviorTreeEditor/Classes/BehaviorTreeGraphNode.h
--- a/Engine/Source/Editor/BehaviorTreeEditor/Classes/BehaviorTreeGraphNode.h
+++ b/Engine/Source/Editor/BehaviorTreeEditor/Classes/BehaviorTreeGraphNode.h
@@ -14,6 +14,8 @@
 
 #include "BehaviorTreeGraphNode.generated.h"
 
+#define UE_API BEHAVIORTREEEDITOR_API
+
 class ISlateStyle;
 class UBehaviorTreeGraph;
 class UEdGraph;
@@ -21,8 +23,8 @@ class UEdGraphSchema;
 class UObject;
 template <typename T> struct TObjectPtr;
 
-UCLASS()
-class BEHAVIORTREEEDITOR_API UBehaviorTreeGraphNode : public UAIGraphNode
+UCLASS(MinimalAPI)
+class UBehaviorTreeGraphNode : public UAIGraphNode
 {
 	GENERATED_UCLASS_BODY()
 
@@ -35,42 +37,42 @@ class BEHAVIORTREEEDITOR_API UBehaviorTreeGraphNode : public UAIGraphNode
 	TArray<TObjectPtr<UBehaviorTreeGraphNode>> Services;
 
 	//~ Begin UEdGraphNode Interface
-	virtual void AllocateDefaultPins() override;
-	virtual FText GetTooltipText() const override;
-	virtual bool CanCreateUnderSpecifiedSchema(const UEdGraphSchema* DesiredSchema) const override;
-	virtual void FindDiffs(class UEdGraphNode* OtherNode, struct FDiffResults& Results) override;
+	UE_API virtual void AllocateDefaultPins() override;
+	UE_API virtual FText GetTooltipText() const override;
+	UE_API virtual bool CanCreateUnderSpecifiedSchema(const UEdGraphSchema* DesiredSchema) const override;
+	UE_API virtual void FindDiffs(class UEdGraphNode* OtherNode, struct FDiffResults& Results) override;
 	//~ End UEdGraphNode Interface
 
 #if WITH_EDITOR
-	virtual void PostEditUndo() override;
+	UE_API virtual void PostEditUndo() override;
 #endif
 
-	virtual FText GetDescription() const override;
-	virtual bool HasErrors() const override;
-	virtual void InitializeInstance() override;
-	virtual void OnSubNodeAdded(UAIGraphNode* SubNode) override;
-	virtual void OnSubNodeRemoved(UAIGraphNode* SubNode) override;
-	virtual void RemoveAllSubNodes() override;
-	virtual int32 FindSubNodeDropIndex(UAIGraphNode* SubNode) const override;
-	virtual void InsertSubNodeAt(UAIGraphNode* SubNode, int32 DropIndex) override;
-	virtual void UpdateErrorMessage() override;
+	UE_API virtual FText GetDescription() const override;
+	UE_API virtual bool HasErrors() const override;
+	UE_API virtual void InitializeInstance() override;
+	UE_API virtual void OnSubNodeAdded(UAIGraphNode* SubNode) override;
+	UE_API virtual void OnSubNodeRemoved(UAIGraphNode* SubNode) override;
+	UE_API virtual void RemoveAllSubNodes() override;
+	UE_API virtual int32 FindSubNodeDropIndex(UAIGraphNode* SubNode) const override;
+	UE_API virtual void InsertSubNodeAt(UAIGraphNode* SubNode, int32 DropIndex) override;
+	UE_API virtual void UpdateErrorMessage() override;
 
 	UE_DEPRECATED(5.4, "Use GetOwnerBehaviorTreeGraph instead.")
-	virtual UBehaviorTreeGraph* GetBehaviorTreeGraph();
-	virtual UBehaviorTreeGraph* GetOwnerBehaviorTreeGraph() const;
+	UE_API virtual UBehaviorTreeGraph* GetBehaviorTreeGraph();
+	UE_API virtual UBehaviorTreeGraph* GetOwnerBehaviorTreeGraph() const;
 
-	virtual FLinearColor GetBackgroundColor(bool bIsActiveForDebugger) const;
+	UE_API virtual FLinearColor GetBackgroundColor(bool bIsActiveForDebugger) const;
 
 	/** check if node can accept breakpoints */
 	virtual bool CanPlaceBreakpoints() const { return false; }
 
-	void ClearDebuggerState();
+	UE_API void ClearDebuggerState();
 
 	/** gets the style set from which GetNameIcon is queried */
-	virtual const ISlateStyle& GetNameIconStyleSet() const;
+	UE_API virtual const ISlateStyle& GetNameIconStyleSet() const;
 
 	/** gets icon resource name for title bar */
-	virtual FName GetNameIcon() const;
+	UE_API virtual FName GetNameIcon() const;
 
 	/** if set, this node was injected from subtree and shouldn't be edited */
 	UPROPERTY()
@@ -148,14 +150,16 @@ class BEHAVIORTREEEDITOR_API UBehaviorTreeGraphNode : public UAIGraphNode
 protected:
 
 	/** creates add decorator... submenu */
-	void CreateAddDecoratorSubMenu(class UToolMenu* Menu, UEdGraph* Graph) const;
+	UE_API void CreateAddDecoratorSubMenu(class UToolMenu* Menu, UEdGraph* Graph) const;
 
 	/** creates add service... submenu */
-	void CreateAddServiceSubMenu(class UToolMenu* Menu, UEdGraph* Graph) const;
+	UE_API void CreateAddServiceSubMenu(class UToolMenu* Menu, UEdGraph* Graph) const;
 
 	/** add right click menu to create subnodes: Decorators */
-	void AddContextMenuActionsDecorators(class UToolMenu* Menu, const FName SectionName, class UGraphNodeContextMenuContext* Context) const;
+	UE_API void AddContextMenuActionsDecorators(class UToolMenu* Menu, const FName SectionName, class UGraphNodeContextMenuContext* Context) const;
 
 	/** add right click menu to create subnodes: Services */
-	void AddContextMenuActionsServices(class UToolMenu* Menu, const FName SectionName, class UGraphNodeContextMenuContext* Context) const;
+	UE_API void AddContextMenuActionsServices(class UToolMenu* Menu, const FName SectionName, class UGraphNodeContextMenuContext* Context) const;
 };
+
+#undef UE_API
diff --git a/Engine/Source/Editor/BehaviorTreeEditor/Classes/BehaviorTreeGraphNode_Composite.h b/Engine/Source/Editor/BehaviorTreeEditor/Classes/BehaviorTreeGraphNode_Composite.h
--- a/Engine/Source/Editor/BehaviorTreeEditor/Classes/BehaviorTreeGraphNode_Composite.h
+++ b/Engine/Source/Editor/BehaviorTreeEditor/Classes/BehaviorTreeGraphNode_Composite.h
@@ -11,27 +11,31 @@
 
 #include "BehaviorTreeGraphNode_Composite.generated.h"
 
+#define UE_API BEHAVIORTREEEDITOR_API
+
 class UObject;
 
-UCLASS()
-class BEHAVIORTREEEDITOR_API UBehaviorTreeGraphNode_Composite : public UBehaviorTreeGraphNode
+UCLASS(MinimalAPI)
+class UBehaviorTreeGraphNode_Composite : public UBehaviorTreeGraphNode
 {
 	GENERATED_UCLASS_BODY()
 	
-	virtual FText GetNodeTitle(ENodeTitleType::Type TitleType) const override;
-	virtual FText GetDescription() const override;
-	virtual FText GetTooltipText() const override;
+	UE_API virtual FText GetNodeTitle(ENodeTitleType::Type TitleType) const override;
+	UE_API virtual FText GetDescription() const override;
+	UE_API virtual FText GetTooltipText() const override;
 	virtual bool RefreshNodeClass() override{ return false; }
 
 	/** Gets a list of actions that can be done to this particular node */
-	virtual void GetNodeContextMenuActions(class UToolMenu* Menu, class UGraphNodeContextMenuContext* Context) const override;
+	UE_API virtual void GetNodeContextMenuActions(class UToolMenu* Menu, class UGraphNodeContextMenuContext* Context) const override;
 
 	/** check if node can accept breakpoints */
 	virtual bool CanPlaceBreakpoints() const override { return true; }
 
-	virtual FLinearColor GetBackgroundColor(bool bIsActiveForDebugger) const override;
+	UE_API virtual FLinearColor GetBackgroundColor(bool bIsActiveForDebugger) const override;
 
 protected:
 
-	virtual void PostPasteNode() override;
+	UE_API virtual void PostPasteNode() override;
 };
+
+#undef UE_API
diff --git a/Engine/Source/Editor/BehaviorTreeEditor/Classes/BehaviorTreeGraphNode_CompositeDecorator.h b/Engine/Source/Editor/BehaviorTreeEditor/Classes/BehaviorTreeGraphNode_CompositeDecorator.h
--- a/Engine/Source/Editor/BehaviorTreeEditor/Classes/BehaviorTreeGraphNode_CompositeDecorator.h
+++ b/Engine/Source/Editor/BehaviorTreeEditor/Classes/BehaviorTreeGraphNode_CompositeDecorator.h
@@ -16,12 +16,14 @@
 
 #include "BehaviorTreeGraphNode_CompositeDecorator.generated.h"
 
+#define UE_API BEHAVIORTREEEDITOR_API
+
 class UBehaviorTreeDecoratorGraph;
 class UEdGraph;
 class UObject;
 
-UCLASS()
-class BEHAVIORTREEEDITOR_API UBehaviorTreeGraphNode_CompositeDecorator : public UBehaviorTreeGraphNode
+UCLASS(MinimalAPI)
+class UBehaviorTreeGraphNode_CompositeDecorator : public UBehaviorTreeGraphNode
 {
 	GENERATED_UCLASS_BODY()
 
@@ -44,45 +46,45 @@ class BEHAVIORTREEEDITOR_API UBehaviorTreeGraphNode_CompositeDecorator : public
 
 	TSubclassOf<UBehaviorTreeDecoratorGraph> GraphClass;
 
-	FString GetNodeTypeDescription() const;
-	virtual FText GetNodeTitle(ENodeTitleType::Type TitleType) const override;
-	virtual void AllocateDefaultPins() override;
-	virtual FText GetDescription() const override;
-	virtual FText GetTooltipText() const override;
-	virtual void PostPlacedNewNode() override;
-	virtual void PostLoad() override;
+	UE_API FString GetNodeTypeDescription() const;
+	UE_API virtual FText GetNodeTitle(ENodeTitleType::Type TitleType) const override;
+	UE_API virtual void AllocateDefaultPins() override;
+	UE_API virtual FText GetDescription() const override;
+	UE_API virtual FText GetTooltipText() const override;
+	UE_API virtual void PostPlacedNewNode() override;
+	UE_API virtual void PostLoad() override;
 	virtual UEdGraph* GetBoundGraph() const override { return BoundGraph; }
-	virtual bool IsSubNode() const override;
-	virtual bool HasErrors() const override;
-	virtual bool RefreshNodeClass() override;
-	virtual void UpdateNodeClassData() override;
+	UE_API virtual bool IsSubNode() const override;
+	UE_API virtual bool HasErrors() const override;
+	UE_API virtual bool RefreshNodeClass() override;
+	UE_API virtual void UpdateNodeClassData() override;
 
-	virtual void PrepareForCopying() override;
-	virtual void PostCopyNode() override;
-	virtual void PostPasteNode() override;
+	UE_API virtual void PrepareForCopying() override;
+	UE_API virtual void PostCopyNode() override;
+	UE_API virtual void PostPasteNode() override;
 
-	int32 SpawnMissingNodes(const TArray<class UBTDecorator*>& NodeInstances, const TArray<struct FBTDecoratorLogic>& Operations, int32 StartIndex);
-	void CollectDecoratorData(TArray<class UBTDecorator*>& NodeInstances, TArray<struct FBTDecoratorLogic>& Operations) const;
-	void SetDecoratorData(class UBTCompositeNode* InParentNode, uint8 InChildIndex);
-	void InitializeDecorator(class UBTDecorator* InnerDecorator);
-	void OnBlackboardUpdate();
-	void OnInnerGraphChanged();
-	void BuildDescription();
-	void UpdateBrokenInstances();
+	UE_API int32 SpawnMissingNodes(const TArray<class UBTDecorator*>& NodeInstances, const TArray<struct FBTDecoratorLogic>& Operations, int32 StartIndex);
+	UE_API void CollectDecoratorData(TArray<class UBTDecorator*>& NodeInstances, TArray<struct FBTDecoratorLogic>& Operations) const;
+	UE_API void SetDecoratorData(class UBTCompositeNode* InParentNode, uint8 InChildIndex);
+	UE_API void InitializeDecorator(class UBTDecorator* InnerDecorator);
+	UE_API void OnBlackboardUpdate();
+	UE_API void OnInnerGraphChanged();
+	UE_API void BuildDescription();
+	UE_API void UpdateBrokenInstances();
 
-	virtual void PostEditChangeProperty(struct FPropertyChangedEvent& PropertyChangedEvent) override;
+	UE_API virtual void PostEditChangeProperty(struct FPropertyChangedEvent& PropertyChangedEvent) override;
 
-	virtual FLinearColor GetBackgroundColor(bool bIsActiveForDebugger) const override;
+	UE_API virtual FLinearColor GetBackgroundColor(bool bIsActiveForDebugger) const override;
 
-	void ResetExecutionRange();
+	UE_API void ResetExecutionRange();
 
 	/** Execution index range of internal nodes, used by debugger */
 	uint16 FirstExecutionIndex;
 	uint16 LastExecutionIndex;
 
 protected:
-	void CreateBoundGraph();
-	virtual void ResetNodeOwner() override;
+	UE_API void CreateBoundGraph();
+	UE_API virtual void ResetNodeOwner() override;
 
 	UPROPERTY()
 	TObjectPtr<class UBTCompositeNode> ParentNodeInstance;
@@ -92,3 +94,5 @@ protected:
 	UPROPERTY()
 	FString CachedDescription;
 };
+
+#undef UE_API
diff --git a/Engine/Source/Editor/BehaviorTreeEditor/Classes/BehaviorTreeGraphNode_Root.h b/Engine/Source/Editor/BehaviorTreeEditor/Classes/BehaviorTreeGraphNode_Root.h
--- a/Engine/Source/Editor/BehaviorTreeEditor/Classes/BehaviorTreeGraphNode_Root.h
+++ b/Engine/Source/Editor/BehaviorTreeEditor/Classes/BehaviorTreeGraphNode_Root.h
@@ -13,34 +13,38 @@
 
 #include "BehaviorTreeGraphNode_Root.generated.h"
 
+#define UE_API BEHAVIORTREEEDITOR_API
+
 class UObject;
 
 /** Root node of this behavior tree, holds Blackboard data */
-UCLASS()
-class BEHAVIORTREEEDITOR_API UBehaviorTreeGraphNode_Root : public UBehaviorTreeGraphNode
+UCLASS(MinimalAPI)
+class UBehaviorTreeGraphNode_Root : public UBehaviorTreeGraphNode
 {
 	GENERATED_UCLASS_BODY()
 
 	UPROPERTY(EditAnywhere, Category="AI|BehaviorTree")
 	TObjectPtr<class UBlackboardData> BlackboardAsset;
 
-	virtual void PostPlacedNewNode() override;
-	virtual void AllocateDefaultPins() override;
+	UE_API virtual void PostPlacedNewNode() override;
+	UE_API virtual void AllocateDefaultPins() override;
 	virtual bool CanDuplicateNode() const override { return false; }
 	virtual bool CanUserDeleteNode() const override{ return false; }
 	virtual bool HasErrors() const override { return false; }
 	virtual bool RefreshNodeClass() override { return false; }
-	virtual FText GetNodeTitle(ENodeTitleType::Type TitleType) const override;
+	UE_API virtual FText GetNodeTitle(ENodeTitleType::Type TitleType) const override;
 	/** gets icon resource name for title bar */
-	virtual FName GetNameIcon() const override;
-	virtual FText GetTooltipText() const override;
+	UE_API virtual FName GetNameIcon() const override;
+	UE_API virtual FText GetTooltipText() const override;
 
-	virtual void PostEditChangeProperty( struct FPropertyChangedEvent& PropertyChangedEvent) override;
-	virtual void PostEditUndo() override;
-	virtual FText GetDescription() const override;
+	UE_API virtual void PostEditChangeProperty( struct FPropertyChangedEvent& PropertyChangedEvent) override;
+	UE_API virtual void PostEditUndo() override;
+	UE_API virtual FText GetDescription() const override;
 
-	virtual FLinearColor GetBackgroundColor(bool bIsActiveForDebugger) const override;
+	UE_API virtual FLinearColor GetBackgroundColor(bool bIsActiveForDebugger) const override;
 
 	/** notify behavior tree about blackboard change */
-	void UpdateBlackboard();
+	UE_API void UpdateBlackboard();
 };
+
+#undef UE_API
diff --git a/Engine/Source/Editor/BehaviorTreeEditor/Classes/BehaviorTreeGraphNode_Task.h b/Engine/Source/Editor/BehaviorTreeEditor/Classes/BehaviorTreeGraphNode_Task.h
--- a/Engine/Source/Editor/BehaviorTreeEditor/Classes/BehaviorTreeGraphNode_Task.h
+++ b/Engine/Source/Editor/BehaviorTreeEditor/Classes/BehaviorTreeGraphNode_Task.h
@@ -23,6 +23,6 @@ class UBehaviorTreeGraphNode_Task : public UBehaviorTreeGraphNode
 	/** Gets a list of actions that can be done to this particular node */
 	BEHAVIORTREEEDITOR_API virtual void GetNodeContextMenuActions(class UToolMenu* Menu, class UGraphNodeContextMenuContext* Context) const override;
 
-	BEHAVIORTREEEDITOR_API virtual bool CanPlaceBreakpoints() const override { return true; }
+	virtual bool CanPlaceBreakpoints() const override { return true; }
 	BEHAVIORTREEEDITOR_API virtual FLinearColor GetBackgroundColor(bool bIsActiveForDebugger) const override;
 };
diff --git a/Engine/Source/Editor/BehaviorTreeEditor/Classes/EdGraphSchema_BehaviorTree.h b/Engine/Source/Editor/BehaviorTreeEditor/Classes/EdGraphSchema_BehaviorTree.h
--- a/Engine/Source/Editor/BehaviorTreeEditor/Classes/EdGraphSchema_BehaviorTree.h
+++ b/Engine/Source/Editor/BehaviorTreeEditor/Classes/EdGraphSchema_BehaviorTree.h
@@ -18,6 +18,8 @@
 
 #include "EdGraphSchema_BehaviorTree.generated.h"
 
+#define UE_API BEHAVIORTREEEDITOR_API
+
 class FSlateRect;
 class UBehaviorTreeGraphNode_CompositeDecorator;
 class UBehaviorTreeGraphNode_Decorator;
@@ -49,26 +51,26 @@ struct FBehaviorTreeSchemaAction_AutoArrange : public FEdGraphSchemaAction
 	//~ End FEdGraphSchemaAction Interface
 };
 
-UCLASS()
-class BEHAVIORTREEEDITOR_API UEdGraphSchema_BehaviorTree : public UAIGraphSchema
+UCLASS(MinimalAPI)
+class UEdGraphSchema_BehaviorTree : public UAIGraphSchema
 {
 	GENERATED_UCLASS_BODY()
 
 	//~ Begin EdGraphSchema Interface
-	virtual void CreateDefaultNodesForGraph(UEdGraph& Graph) const override;
-	virtual void GetGraphContextActions(FGraphContextMenuBuilder& ContextMenuBuilder) const override;
-	virtual void GetContextMenuActions(class UToolMenu* Menu, class UGraphNodeContextMenuContext* Context) const override;
-	virtual const FPinConnectionResponse CanCreateConnection(const UEdGraphPin* A, const UEdGraphPin* B) const override;
-	virtual const FPinConnectionResponse CanMergeNodes(const UEdGraphNode* A, const UEdGraphNode* B) const override;
-	virtual FLinearColor GetPinTypeColor(const FEdGraphPinType& PinType) const override;
-	virtual class FConnectionDrawingPolicy* CreateConnectionDrawingPolicy(int32 InBackLayerID, int32 InFrontLayerID, float InZoomFactor, const FSlateRect& InClippingRect, class FSlateWindowElementList& InDrawElements, class UEdGraph* InGraphObj) const override;
-	virtual bool IsCacheVisualizationOutOfDate(int32 InVisualizationCacheID) const override;
-	virtual int32 GetCurrentVisualizationCacheID() const override;
-	virtual void ForceVisualizationCacheClear() const override;
+	UE_API virtual void CreateDefaultNodesForGraph(UEdGraph& Graph) const override;
+	UE_API virtual void GetGraphContextActions(FGraphContextMenuBuilder& ContextMenuBuilder) const override;
+	UE_API virtual void GetContextMenuActions(class UToolMenu* Menu, class UGraphNodeContextMenuContext* Context) const override;
+	UE_API virtual const FPinConnectionResponse CanCreateConnection(const UEdGraphPin* A, const UEdGraphPin* B) const override;
+	UE_API virtual const FPinConnectionResponse CanMergeNodes(const UEdGraphNode* A, const UEdGraphNode* B) const override;
+	UE_API virtual FLinearColor GetPinTypeColor(const FEdGraphPinType& PinType) const override;
+	UE_API virtual class FConnectionDrawingPolicy* CreateConnectionDrawingPolicy(int32 InBackLayerID, int32 InFrontLayerID, float InZoomFactor, const FSlateRect& InClippingRect, class FSlateWindowElementList& InDrawElements, class UEdGraph* InGraphObj) const override;
+	UE_API virtual bool IsCacheVisualizationOutOfDate(int32 InVisualizationCacheID) const override;
+	UE_API virtual int32 GetCurrentVisualizationCacheID() const override;
+	UE_API virtual void ForceVisualizationCacheClear() const override;
 	//~ End EdGraphSchema Interface
 
-	virtual void GetGraphNodeContextActions(FGraphContextMenuBuilder& ContextMenuBuilder, int32 SubNodeFlags) const override;
-	virtual void GetSubNodeClasses(int32 SubNodeFlags, TArray<FGraphNodeClassData>& ClassData, UClass*& GraphNodeClass) const override;
+	UE_API virtual void GetGraphNodeContextActions(FGraphContextMenuBuilder& ContextMenuBuilder, int32 SubNodeFlags) const override;
+	UE_API virtual void GetSubNodeClasses(int32 SubNodeFlags, TArray<FGraphNodeClassData>& ClassData, UClass*& GraphNodeClass) const override;
 
 	TSubclassOf<UBehaviorTreeGraphNode_CompositeDecorator> CompositeDecoratorClass;
 	TSubclassOf<UBehaviorTreeGraphNode_Decorator> DecoratorClass;
@@ -76,11 +78,12 @@ class BEHAVIORTREEEDITOR_API UEdGraphSchema_BehaviorTree : public UAIGraphSchema
 	TSubclassOf<UBehaviorTreeGraphNode_Service> ServiceClass;
 
 protected:
-	virtual FGraphNodeClassHelper& GetClassCache() const;
-	virtual bool IsNodeSubtreeTask(const FGraphNodeClassData& NodeClass) const;
+	UE_API virtual FGraphNodeClassHelper& GetClassCache() const;
+	UE_API virtual bool IsNodeSubtreeTask(const FGraphNodeClassData& NodeClass) const;
 
 private:
 	// ID for checking dirty status of node titles against, increases whenever 
-	static int32 CurrentCacheRefreshID;
+	static UE_API int32 CurrentCacheRefreshID;
 };
 
+#undef UE_API
diff --git a/Engine/Source/Editor/BehaviorTreeEditor/Classes/EdGraphSchema_BehaviorTreeDecorator.h b/Engine/Source/Editor/BehaviorTreeEditor/Classes/EdGraphSchema_BehaviorTreeDecorator.h
--- a/Engine/Source/Editor/BehaviorTreeEditor/Classes/EdGraphSchema_BehaviorTreeDecorator.h
+++ b/Engine/Source/Editor/BehaviorTreeEditor/Classes/EdGraphSchema_BehaviorTreeDecorator.h
@@ -20,6 +20,8 @@
 
 #include "EdGraphSchema_BehaviorTreeDecorator.generated.h"
 
+#define UE_API BEHAVIORTREEEDITOR_API
+
 class UEdGraph;
 class UEdGraphNode;
 class UEdGraphPin;
@@ -63,35 +65,37 @@ struct FDecoratorSchemaAction_NewNode : public FEdGraphSchemaAction
 	}
 };
 
-UCLASS()
-class BEHAVIORTREEEDITOR_API UEdGraphSchema_BehaviorTreeDecorator : public UEdGraphSchema
+UCLASS(MinimalAPI)
+class UEdGraphSchema_BehaviorTreeDecorator : public UEdGraphSchema
 {
 	GENERATED_UCLASS_BODY()
 
 	UPROPERTY()
 	FString PC_Boolean;
 
-	void AddPin(class UEdGraphNode* InGraphNode);
-	void RemovePin(class UEdGraphPin* InGraphPin);
+	UE_API void AddPin(class UEdGraphNode* InGraphNode);
+	UE_API void RemovePin(class UEdGraphPin* InGraphPin);
 
 	//~ Begin EdGraphSchema Interface
-	virtual void CreateDefaultNodesForGraph(UEdGraph& Graph) const override;
-	virtual void GetGraphContextActions(FGraphContextMenuBuilder& ContextMenuBuilder) const override;
-	virtual void GetContextMenuActions(class UToolMenu* Menu, class UGraphNodeContextMenuContext* Context) const override;
+	UE_API virtual void CreateDefaultNodesForGraph(UEdGraph& Graph) const override;
+	UE_API virtual void GetGraphContextActions(FGraphContextMenuBuilder& ContextMenuBuilder) const override;
+	UE_API virtual void GetContextMenuActions(class UToolMenu* Menu, class UGraphNodeContextMenuContext* Context) const override;
 	virtual FName GetParentContextMenuName() const override { return NAME_None; }
-	virtual const FPinConnectionResponse CanCreateConnection(const UEdGraphPin* A, const UEdGraphPin* B) const override;
-	virtual FLinearColor GetPinTypeColor(const FEdGraphPinType& PinType) const override;
-	virtual bool ShouldHidePinDefaultValue(UEdGraphPin* Pin) const override;
-	virtual bool IsCacheVisualizationOutOfDate(int32 InVisualizationCacheID) const override;
-	virtual int32 GetCurrentVisualizationCacheID() const override;
-	virtual void ForceVisualizationCacheClear() const override;
+	UE_API virtual const FPinConnectionResponse CanCreateConnection(const UEdGraphPin* A, const UEdGraphPin* B) const override;
+	UE_API virtual FLinearColor GetPinTypeColor(const FEdGraphPinType& PinType) const override;
+	UE_API virtual bool ShouldHidePinDefaultValue(UEdGraphPin* Pin) const override;
+	UE_API virtual bool IsCacheVisualizationOutOfDate(int32 InVisualizationCacheID) const override;
+	UE_API virtual int32 GetCurrentVisualizationCacheID() const override;
+	UE_API virtual void ForceVisualizationCacheClear() const override;
 	//~ End EdGraphSchema Interface
 
-	static TSharedPtr<FDecoratorSchemaAction_NewNode> AddNewDecoratorAction(FGraphContextMenuBuilder& ContextMenuBuilder, const FText& Category, const FText& MenuDesc, const FText& Tooltip);
+	static UE_API TSharedPtr<FDecoratorSchemaAction_NewNode> AddNewDecoratorAction(FGraphContextMenuBuilder& ContextMenuBuilder, const FText& Category, const FText& MenuDesc, const FText& Tooltip);
 
 protected:
-	virtual FGraphNodeClassHelper& GetClassCache() const;
+	UE_API virtual FGraphNodeClassHelper& GetClassCache() const;
 
 private:
 	inline static int32 CurrentCacheRefreshID = 0;
 };
+
+#undef UE_API
diff --git a/Engine/Source/Editor/BehaviorTreeEditor/Classes/ValueOrBBKeyDetails.h b/Engine/Source/Editor/BehaviorTreeEditor/Classes/ValueOrBBKeyDetails.h
--- a/Engine/Source/Editor/BehaviorTreeEditor/Classes/ValueOrBBKeyDetails.h
+++ b/Engine/Source/Editor/BehaviorTreeEditor/Classes/ValueOrBBKeyDetails.h
@@ -9,6 +9,8 @@
 #include "Templates/SharedPointer.h"
 #include "UObject/NameTypes.h"
 
+#define UE_API BEHAVIORTREEEDITOR_API
+
 namespace ESelectInfo
 {
 	enum Type : int;
@@ -20,16 +22,16 @@ class SWidget;
 class UBlackboardData;
 struct FValueOrBlackboardKeyBase;
 
-class BEHAVIORTREEEDITOR_API FValueOrBBKeyDetails : public IPropertyTypeCustomization
+class FValueOrBBKeyDetails : public IPropertyTypeCustomization
 {
 public:
 	/** Makes a new instance of this detail layout class for a specific detail view requesting it */
-	static TSharedRef<IPropertyTypeCustomization> MakeInstance();
+	static UE_API TSharedRef<IPropertyTypeCustomization> MakeInstance();
 
-	virtual void CustomizeHeader(TSharedRef<class IPropertyHandle> StructPropertyHandle, class FDetailWidgetRow& HeaderRow, IPropertyTypeCustomizationUtils& StructCustomizationUtils) override;
-	virtual void CustomizeChildren(TSharedRef<class IPropertyHandle> StructPropertyHandle, class IDetailChildrenBuilder& StructBuilder, IPropertyTypeCustomizationUtils& StructCustomizationUtils) override;
+	UE_API virtual void CustomizeHeader(TSharedRef<class IPropertyHandle> StructPropertyHandle, class FDetailWidgetRow& HeaderRow, IPropertyTypeCustomizationUtils& StructCustomizationUtils) override;
+	UE_API virtual void CustomizeChildren(TSharedRef<class IPropertyHandle> StructPropertyHandle, class IDetailChildrenBuilder& StructBuilder, IPropertyTypeCustomizationUtils& StructCustomizationUtils) override;
 
-	bool CanEditDefaultValue() const;
+	UE_API bool CanEditDefaultValue() const;
 
 protected:
 	TSharedPtr<IPropertyHandle> StructProperty;
@@ -38,15 +40,15 @@ protected:
 	TSharedPtr<IPropertyUtilities> CachedUtils;
 	TArray<FName> MatchingKeys;
 
-	virtual void ValidateData();
-	virtual TSharedRef<SWidget> CreateDefaultValueWidget();
+	UE_API virtual void ValidateData();
+	UE_API virtual TSharedRef<SWidget> CreateDefaultValueWidget();
 
-	void GetMatchingKeys(TArray<FName>& OutNames);
-	bool HasAccessToBlackboard() const;
-	TSharedRef<SWidget> OnGetKeyNames();
-	void OnKeyChanged(int32 Index);
-	FText GetKeyDesc() const;
-	const FValueOrBlackboardKeyBase* GetDataPtr() const;
+	UE_API void GetMatchingKeys(TArray<FName>& OutNames);
+	UE_API bool HasAccessToBlackboard() const;
+	UE_API TSharedRef<SWidget> OnGetKeyNames();
+	UE_API void OnKeyChanged(int32 Index);
+	UE_API FText GetKeyDesc() const;
+	UE_API const FValueOrBlackboardKeyBase* GetDataPtr() const;
 };
 
 class FValueOrBBKeyDetails_Class : public FValueOrBBKeyDetails
@@ -127,10 +129,12 @@ protected:
 	TSharedPtr<IPropertyHandle> EditDefaultsOnlyProperty;
 };
 
-class BEHAVIORTREEEDITOR_API FValueOrBBKeyDetails_WithChild : public FValueOrBBKeyDetails
+class FValueOrBBKeyDetails_WithChild : public FValueOrBBKeyDetails
 {
 public:
 	/** Makes a new instance of this detail layout class for a specific detail view requesting it */
-	static TSharedRef<IPropertyTypeCustomization> MakeInstance();
-	virtual void CustomizeChildren(TSharedRef<class IPropertyHandle> StructPropertyHandle, class IDetailChildrenBuilder& StructBuilder, IPropertyTypeCustomizationUtils& StructCustomizationUtils) override;
-};
\ No newline at end of file
+	static UE_API TSharedRef<IPropertyTypeCustomization> MakeInstance();
+	UE_API virtual void CustomizeChildren(TSharedRef<class IPropertyHandle> StructPropertyHandle, class IDetailChildrenBuilder& StructBuilder, IPropertyTypeCustomizationUtils& StructCustomizationUtils) override;
+};
+
+#undef UE_API
diff --git a/Engine/Source/Editor/BehaviorTreeEditor/Private/BehaviorTreeConnectionDrawingPolicy.h b/Engine/Source/Editor/BehaviorTreeEditor/Private/BehaviorTreeConnectionDrawingPolicy.h
--- a/Engine/Source/Editor/BehaviorTreeEditor/Private/BehaviorTreeConnectionDrawingPolicy.h
+++ b/Engine/Source/Editor/BehaviorTreeEditor/Private/BehaviorTreeConnectionDrawingPolicy.h
@@ -5,6 +5,8 @@
 #include "AIGraphConnectionDrawingPolicy.h"
 #include "HAL/Platform.h"
 
+#define UE_API BEHAVIORTREEEDITOR_API
+
 class FSlateRect;
 class FSlateWindowElementList;
 class UEdGraph;
@@ -12,13 +14,15 @@ class UEdGraphPin;
 struct FConnectionParams;
 
 // This class draws the connections for an UEdGraph with a behavior tree schema
-class BEHAVIORTREEEDITOR_API FBehaviorTreeConnectionDrawingPolicy : public FAIGraphConnectionDrawingPolicy
+class FBehaviorTreeConnectionDrawingPolicy : public FAIGraphConnectionDrawingPolicy
 {
 public:
 	//
-	FBehaviorTreeConnectionDrawingPolicy(int32 InBackLayerID, int32 InFrontLayerID, float ZoomFactor, const FSlateRect& InClippingRect, FSlateWindowElementList& InDrawElements, UEdGraph* InGraphObj);
+	UE_API FBehaviorTreeConnectionDrawingPolicy(int32 InBackLayerID, int32 InFrontLayerID, float ZoomFactor, const FSlateRect& InClippingRect, FSlateWindowElementList& InDrawElements, UEdGraph* InGraphObj);
 
 	// FConnectionDrawingPolicy interface 
-	virtual void DetermineWiringStyle(UEdGraphPin* OutputPin, UEdGraphPin* InputPin, /*inout*/ FConnectionParams& Params) override;
+	UE_API virtual void DetermineWiringStyle(UEdGraphPin* OutputPin, UEdGraphPin* InputPin, /*inout*/ FConnectionParams& Params) override;
 	// End of FConnectionDrawingPolicy interface
 };
+
+#undef UE_API
diff --git a/Engine/Source/Editor/BehaviorTreeEditor/Public/AssetDefinition_BehaviorTree.h b/Engine/Source/Editor/BehaviorTreeEditor/Public/AssetDefinition_BehaviorTree.h
--- a/Engine/Source/Editor/BehaviorTreeEditor/Public/AssetDefinition_BehaviorTree.h
+++ b/Engine/Source/Editor/BehaviorTreeEditor/Public/AssetDefinition_BehaviorTree.h
@@ -5,25 +5,29 @@
 #include "AssetDefinitionDefault.h"
 #include "AssetDefinition_BehaviorTree.generated.h"
 
+#define UE_API BEHAVIORTREEEDITOR_API
+
 class UBehaviorTree;
 
-UCLASS()
-class BEHAVIORTREEEDITOR_API UAssetDefinition_BehaviorTree : public UAssetDefinitionDefault
+UCLASS(MinimalAPI)
+class UAssetDefinition_BehaviorTree : public UAssetDefinitionDefault
 {
 	GENERATED_BODY()
 
 protected:
 	// UAssetDefinition Begin
-	virtual FText GetAssetDisplayName() const override;
-	virtual FLinearColor GetAssetColor() const override;
-	virtual TSoftClassPtr<UObject> GetAssetClass() const override;
-	virtual TConstArrayView<FAssetCategoryPath> GetAssetCategories() const override;
-	virtual EAssetCommandResult OpenAssets(const FAssetOpenArgs& OpenArgs) const override;
-	virtual EAssetCommandResult PerformAssetDiff(const FAssetDiffArgs& DiffArgs) const override;
+	UE_API virtual FText GetAssetDisplayName() const override;
+	UE_API virtual FLinearColor GetAssetColor() const override;
+	UE_API virtual TSoftClassPtr<UObject> GetAssetClass() const override;
+	UE_API virtual TConstArrayView<FAssetCategoryPath> GetAssetCategories() const override;
+	UE_API virtual EAssetCommandResult OpenAssets(const FAssetOpenArgs& OpenArgs) const override;
+	UE_API virtual EAssetCommandResult PerformAssetDiff(const FAssetDiffArgs& DiffArgs) const override;
 	// UAssetDefinition End
 
 private:
 
 	/* Called to open the Behavior Tree defaults view, this opens whatever text diff tool the user has */
-	void OpenInDefaults(class UBehaviorTree* OldBehaviorTree, class UBehaviorTree* NewBehaviorTree) const;
+	UE_API void OpenInDefaults(class UBehaviorTree* OldBehaviorTree, class UBehaviorTree* NewBehaviorTree) const;
 };
+
+#undef UE_API
diff --git a/Engine/Source/Editor/BehaviorTreeEditor/Public/BehaviorTreeEditor.h b/Engine/Source/Editor/BehaviorTreeEditor/Public/BehaviorTreeEditor.h
--- a/Engine/Source/Editor/BehaviorTreeEditor/Public/BehaviorTreeEditor.h
+++ b/Engine/Source/Editor/BehaviorTreeEditor/Public/BehaviorTreeEditor.h
@@ -19,6 +19,8 @@
 #include "UObject/NameTypes.h"
 #include "UObject/WeakObjectPtrTemplates.h"
 
+#define UE_API BEHAVIORTREEEDITOR_API
+
 class FDocumentTabFactory;
 class FDocumentTracker;
 class FObjectPostSaveContext;
@@ -37,233 +39,233 @@ class UPackage;
 struct FBlackboardEntry;
 struct FPropertyChangedEvent;
 
-class BEHAVIORTREEEDITOR_API FBehaviorTreeEditor : public IBehaviorTreeEditor, public FAIGraphEditor, public FNotifyHook
+class FBehaviorTreeEditor : public IBehaviorTreeEditor, public FAIGraphEditor, public FNotifyHook
 {
 public:
-	FBehaviorTreeEditor();
+	UE_API FBehaviorTreeEditor();
 	/** Destructor */
-	virtual ~FBehaviorTreeEditor();
+	UE_API virtual ~FBehaviorTreeEditor();
 
-	virtual void RegisterTabSpawners(const TSharedRef<class FTabManager>& TabManager) override;
+	UE_API virtual void RegisterTabSpawners(const TSharedRef<class FTabManager>& TabManager) override;
 
-	virtual void InitBehaviorTreeEditor(const EToolkitMode::Type Mode, const TSharedPtr<class IToolkitHost>& InitToolkitHost, UObject* InObject);
+	UE_API virtual void InitBehaviorTreeEditor(const EToolkitMode::Type Mode, const TSharedPtr<class IToolkitHost>& InitToolkitHost, UObject* InObject);
 
 	//~ Begin IToolkit Interface
-	virtual FName GetToolkitFName() const override;
-	virtual FText GetBaseToolkitName() const override;
-	virtual FString GetWorldCentricTabPrefix() const override;
-	virtual FLinearColor GetWorldCentricTabColorScale() const override;
-	virtual FText GetToolkitName() const override;
-	virtual FText GetToolkitToolTipText() const override;
+	UE_API virtual FName GetToolkitFName() const override;
+	UE_API virtual FText GetBaseToolkitName() const override;
+	UE_API virtual FString GetWorldCentricTabPrefix() const override;
+	UE_API virtual FLinearColor GetWorldCentricTabColorScale() const override;
+	UE_API virtual FText GetToolkitName() const override;
+	UE_API virtual FText GetToolkitToolTipText() const override;
 	//~ End IToolkit Interface
 
 	//~ Begin IBehaviorTreeEditor Interface
-	virtual void InitializeDebuggerState(class FBehaviorTreeDebugger* ParentDebugger) const override;
-	virtual UEdGraphNode* FindInjectedNode(int32 Index) const override;
-	virtual void DoubleClickNode(class UEdGraphNode* Node) override;
-	virtual void FocusAttentionOnNode(UEdGraphNode* Node) override;
-	virtual void FocusWindow(UObject* ObjectToFocusOn = NULL) override;
+	UE_API virtual void InitializeDebuggerState(class FBehaviorTreeDebugger* ParentDebugger) const override;
+	UE_API virtual UEdGraphNode* FindInjectedNode(int32 Index) const override;
+	UE_API virtual void DoubleClickNode(class UEdGraphNode* Node) override;
+	UE_API virtual void FocusAttentionOnNode(UEdGraphNode* Node) override;
+	UE_API virtual void FocusWindow(UObject* ObjectToFocusOn = NULL) override;
 	//~ End IBehaviorTreeEditor Interface
 
 	//~ Begin IAssetEditorInstance Interface
-	virtual bool IncludeAssetInRestoreOpenAssetsPrompt(UObject* Asset) const override;
+	UE_API virtual bool IncludeAssetInRestoreOpenAssetsPrompt(UObject* Asset) const override;
 	//~ End IAssetEditorInstance Interface
 
 	//~ Begin FEditorUndoClient Interface
-	virtual void PostUndo(bool bSuccess) override;
-	virtual void PostRedo(bool bSuccess) override;
+	UE_API virtual void PostUndo(bool bSuccess) override;
+	UE_API virtual void PostRedo(bool bSuccess) override;
 	// End of FEditorUndoClient
 
 	//~ Begin FNotifyHook Interface
-	virtual void NotifyPostChange( const FPropertyChangedEvent& PropertyChangedEvent, FProperty* PropertyThatChanged ) override;
+	UE_API virtual void NotifyPostChange( const FPropertyChangedEvent& PropertyChangedEvent, FProperty* PropertyThatChanged ) override;
 	// End of FNotifyHook
 
 	// Delegates
-	void OnNodeDoubleClicked(class UEdGraphNode* Node);
-	void OnGraphEditorFocused(const TSharedRef<SGraphEditor>& InGraphEditor);
-	void OnNodeTitleCommitted(const FText& NewText, ETextCommit::Type CommitInfo, UEdGraphNode* NodeBeingChanged);
-
-	void OnAddInputPin();
-	bool CanAddInputPin() const;
-	void OnRemoveInputPin();
-	bool CanRemoveInputPin() const;
-
-	void OnEnableBreakpoint();
-	bool CanEnableBreakpoint() const;
-	void OnToggleBreakpoint();
-	bool CanToggleBreakpoint() const;
-	void OnDisableBreakpoint();
-	bool CanDisableBreakpoint() const;
-	void OnAddBreakpoint();
-	bool CanAddBreakpoint() const;
-	void OnRemoveBreakpoint();
-	bool CanRemoveBreakpoint() const;
-
-	void SearchTree();
-	bool CanSearchTree() const;
-
-	void JumpToNode(const UEdGraphNode* Node);
-
-	bool IsPropertyEditable() const;
-	void OnPackageSaved(const FString& PackageFileName, UPackage* Package, FObjectPostSaveContext ObjectSaveContext);
-	void OnFinishedChangingProperties(const FPropertyChangedEvent& PropertyChangedEvent);
-	virtual void OnClassListUpdated() override;
-
-	void UpdateToolbar();
-	bool IsDebuggerReady() const;
+	UE_API void OnNodeDoubleClicked(class UEdGraphNode* Node);
+	UE_API void OnGraphEditorFocused(const TSharedRef<SGraphEditor>& InGraphEditor);
+	UE_API void OnNodeTitleCommitted(const FText& NewText, ETextCommit::Type CommitInfo, UEdGraphNode* NodeBeingChanged);
+
+	UE_API void OnAddInputPin();
+	UE_API bool CanAddInputPin() const;
+	UE_API void OnRemoveInputPin();
+	UE_API bool CanRemoveInputPin() const;
+
+	UE_API void OnEnableBreakpoint();
+	UE_API bool CanEnableBreakpoint() const;
+	UE_API void OnToggleBreakpoint();
+	UE_API bool CanToggleBreakpoint() const;
+	UE_API void OnDisableBreakpoint();
+	UE_API bool CanDisableBreakpoint() const;
+	UE_API void OnAddBreakpoint();
+	UE_API bool CanAddBreakpoint() const;
+	UE_API void OnRemoveBreakpoint();
+	UE_API bool CanRemoveBreakpoint() const;
+
+	UE_API void SearchTree();
+	UE_API bool CanSearchTree() const;
+
+	UE_API void JumpToNode(const UEdGraphNode* Node);
+
+	UE_API bool IsPropertyEditable() const;
+	UE_API void OnPackageSaved(const FString& PackageFileName, UPackage* Package, FObjectPostSaveContext ObjectSaveContext);
+	UE_API void OnFinishedChangingProperties(const FPropertyChangedEvent& PropertyChangedEvent);
+	UE_API virtual void OnClassListUpdated() override;
+
+	UE_API void UpdateToolbar();
+	UE_API bool IsDebuggerReady() const;
 
 	/** Get whether the debugger is currently running and the PIE session is paused */
-	bool IsDebuggerPaused() const;
+	UE_API bool IsDebuggerPaused() const;
 
 	/** Get whether we can edit the tree/blackboard with the debugger active */
-	bool CanEditWithDebuggerActive() const;
+	UE_API bool CanEditWithDebuggerActive() const;
 
-	TSharedRef<class SWidget> OnGetDebuggerActorsMenu();
-	void OnDebuggerActorSelected(TWeakObjectPtr<UBehaviorTreeComponent> InstanceToDebug);
-	FText GetDebuggerActorDesc() const;
-	FGraphAppearanceInfo GetGraphAppearance() const;
-	bool InEditingMode(bool bGraphIsEditable) const;
+	UE_API TSharedRef<class SWidget> OnGetDebuggerActorsMenu();
+	UE_API void OnDebuggerActorSelected(TWeakObjectPtr<UBehaviorTreeComponent> InstanceToDebug);
+	UE_API FText GetDebuggerActorDesc() const;
+	UE_API FGraphAppearanceInfo GetGraphAppearance() const;
+	UE_API bool InEditingMode(bool bGraphIsEditable) const;
 
-	void DebuggerSwitchAsset(UBehaviorTree* NewAsset);
-	void DebuggerUpdateGraph();
+	UE_API void DebuggerSwitchAsset(UBehaviorTree* NewAsset);
+	UE_API void DebuggerUpdateGraph();
 
-	EVisibility GetDebuggerDetailsVisibility() const;
-	EVisibility GetRangeLowerVisibility() const;
-	EVisibility GetRangeSelfVisibility() const;
-	EVisibility GetInjectedNodeVisibility() const;
-	EVisibility GetRootLevelNodeVisibility() const;
+	UE_API EVisibility GetDebuggerDetailsVisibility() const;
+	UE_API EVisibility GetRangeLowerVisibility() const;
+	UE_API EVisibility GetRangeSelfVisibility() const;
+	UE_API EVisibility GetInjectedNodeVisibility() const;
+	UE_API EVisibility GetRootLevelNodeVisibility() const;
 
-	TWeakPtr<SGraphEditor> GetFocusedGraphPtr() const;
+	UE_API TWeakPtr<SGraphEditor> GetFocusedGraphPtr() const;
 
 	/** Check whether the behavior tree mode can be accessed (i.e whether we have a valid tree to edit) */
-	bool CanAccessBehaviorTreeMode() const;
+	UE_API bool CanAccessBehaviorTreeMode() const;
 
 	/** Check whether the blackboard mode can be accessed (i.e whether we have a valid blackboard to edit) */
-	bool CanAccessBlackboardMode() const;
+	UE_API bool CanAccessBlackboardMode() const;
 
 	/** 
 	 * Get the localized text to display for the specified mode 
 	 * @param	InMode	The mode to display
 	 * @return the localized text representation of the mode
 	 */
-	static FText GetLocalizedMode(FName InMode);
+	static UE_API FText GetLocalizedMode(FName InMode);
 
 	/** Access the toolbar builder for this editor */
 	TSharedPtr<class FBehaviorTreeEditorToolbar> GetToolbarBuilder() { return ToolbarBuilder; }
 
 	/** Get the behavior tree we are editing (if any) */
-	UBehaviorTree* GetBehaviorTree() const;
+	UE_API UBehaviorTree* GetBehaviorTree() const;
 
 	/** Get the blackboard we are editing (if any) */
-	UBlackboardData* GetBlackboardData() const;
+	UE_API UBlackboardData* GetBlackboardData() const;
 
 	/** Spawns the tab with the update graph inside */
-	TSharedRef<SWidget> SpawnProperties();
+	UE_API TSharedRef<SWidget> SpawnProperties();
 
 	/** Spawns the search tab */
-	TSharedRef<SWidget> SpawnSearch();
+	UE_API TSharedRef<SWidget> SpawnSearch();
 
 	/** Spawn blackboard details tab */
-	TSharedRef<SWidget> SpawnBlackboardDetails();
+	UE_API TSharedRef<SWidget> SpawnBlackboardDetails();
 
 	/** Spawn blackboard view tab */
-	TSharedRef<SWidget> SpawnBlackboardView();
+	UE_API TSharedRef<SWidget> SpawnBlackboardView();
 
 	/** Spawn blackboard editor tab */
-	TSharedRef<SWidget> SpawnBlackboardEditor();
+	UE_API TSharedRef<SWidget> SpawnBlackboardEditor();
 
 	// @todo This is a hack for now until we reconcile the default toolbar with application modes [duplicated from counterpart in Blueprint Editor]
-	void RegisterToolbarTab(const TSharedRef<class FTabManager>& TabManager);
+	UE_API void RegisterToolbarTab(const TSharedRef<class FTabManager>& TabManager);
 
 	/** Restores the behavior tree graph we were editing or creates a new one if none is available */
-	void RestoreBehaviorTree();
+	UE_API void RestoreBehaviorTree();
 
 	/** Save the graph state for later editing */
-	void SaveEditedObjectState();
+	UE_API void SaveEditedObjectState();
 
 	/** Delegate handler for selection in the blackboard entry list */
-	void HandleBlackboardEntrySelected(const FBlackboardEntry* BlackboardEntry, bool bIsInherited);
+	UE_API void HandleBlackboardEntrySelected(const FBlackboardEntry* BlackboardEntry, bool bIsInherited);
 
 	/** Delegate handler used to retrieve current blackboard selection */
-	int32 HandleGetSelectedBlackboardItemIndex(bool& bOutIsInherited);
+	UE_API int32 HandleGetSelectedBlackboardItemIndex(bool& bOutIsInherited);
 
 	/** Delegate handler for displaying debugger values */
-	FText HandleGetDebugKeyValue(const FName& InKeyName, bool bUseCurrentState) const;
+	UE_API FText HandleGetDebugKeyValue(const FName& InKeyName, bool bUseCurrentState) const;
 
 	/** Delegate handler for retrieving timestamp to display */
-	double HandleGetDebugTimeStamp(bool bUseCurrentState) const;
+	UE_API double HandleGetDebugTimeStamp(bool bUseCurrentState) const;
 
 	/** Delegate handler for when the debugged blackboard changes */
-	void HandleDebuggedBlackboardChanged(UBlackboardData* InObject);
+	UE_API void HandleDebuggedBlackboardChanged(UBlackboardData* InObject);
 
 	/** Delegate handler for determining whether to display the current state */
-	bool HandleGetDisplayCurrentState() const;
+	UE_API bool HandleGetDisplayCurrentState() const;
 
 	/** Check whether blackboard mode is current */
-	bool HandleIsBlackboardModeActive() const;
+	UE_API bool HandleIsBlackboardModeActive() const;
 
 	/** Get the currently selected blackboard entry */
-	void GetBlackboardSelectionInfo(int32& OutSelectionIndex, bool& bOutIsInherited) const;
+	UE_API void GetBlackboardSelectionInfo(int32& OutSelectionIndex, bool& bOutIsInherited) const;
 
 	/** Check to see if we can create a new task node */
-	bool CanCreateNewTask() const;
+	UE_API bool CanCreateNewTask() const;
 
 	/** Check to see if we can create a new decorator node */
-	bool CanCreateNewDecorator() const;
+	UE_API bool CanCreateNewDecorator() const;
 
 	/** Check to see if we can create a new service node */
-	bool CanCreateNewService() const;
+	UE_API bool CanCreateNewService() const;
 
 	/** Create the menu used to make a new task node */
-	TSharedRef<SWidget> HandleCreateNewTaskMenu() const;
+	UE_API TSharedRef<SWidget> HandleCreateNewTaskMenu() const;
 
 	/** Create the menu used to make a new decorator */
-	TSharedRef<SWidget> HandleCreateNewDecoratorMenu() const;
+	UE_API TSharedRef<SWidget> HandleCreateNewDecoratorMenu() const;
 
 	/** Create the menu used to make a new service */
-	TSharedRef<SWidget> HandleCreateNewServiceMenu() const;
+	UE_API TSharedRef<SWidget> HandleCreateNewServiceMenu() const;
 
 	/** Handler for when a node class is picked */
-	void HandleNewNodeClassPicked(UClass* InClass) const;
+	UE_API void HandleNewNodeClassPicked(UClass* InClass) const;
 
 	/** Create a new task from UBTTask_BlueprintBase */
-	void CreateNewTask() const;
+	UE_API void CreateNewTask() const;
 
 	/** Whether the single button to create a new Blueprint-based task is visible */
-	bool IsNewTaskButtonVisible() const;
+	UE_API bool IsNewTaskButtonVisible() const;
 
 	/** Whether the combo button to create a new Blueprint-based task from all available base classes is visible */
-	bool IsNewTaskComboVisible() const;
+	UE_API bool IsNewTaskComboVisible() const;
 
 	/** Create a new decorator from UBTDecorator_BlueprintBase */
-	void CreateNewDecorator() const;
+	UE_API void CreateNewDecorator() const;
 
 	/** Whether the single button to create a new Blueprint-based decorator is visible */
-	bool IsNewDecoratorButtonVisible() const;
+	UE_API bool IsNewDecoratorButtonVisible() const;
 
 	/** Whether the combo button to create a new Blueprint-based decorator from all available base classes is visible */
-	bool IsNewDecoratorComboVisible() const;
+	UE_API bool IsNewDecoratorComboVisible() const;
 
 	/** Create a new service from UBTService_BlueprintBase */
-	void CreateNewService() const;
+	UE_API void CreateNewService() const;
 
 	/** Whether the single button to create a new Blueprint-based service is visible */
-	bool IsNewServiceButtonVisible() const;
+	UE_API bool IsNewServiceButtonVisible() const;
 
 	/** Whether the combo button to create a new Blueprint-based service from all available base classes is visible */
-	bool IsNewServiceComboVisible() const;
+	UE_API bool IsNewServiceComboVisible() const;
 
 	/** Create a new Blackboard alongside the currently-edited behavior tree */
-	void CreateNewBlackboard();
+	UE_API void CreateNewBlackboard();
 
 	/** Whether we can currently create a new Blackboard */
-	bool CanCreateNewBlackboard() const;
+	UE_API bool CanCreateNewBlackboard() const;
 
 protected:
 	/** Called when "Save" is clicked for this asset */
-	virtual void SaveAsset_Execute() override;
+	UE_API virtual void SaveAsset_Execute() override;
 
-	void SetToolbarCreateActionsEnabled(bool bActionsEnabled);
+	UE_API void SetToolbarCreateActionsEnabled(bool bActionsEnabled);
 
 	TSubclassOf<UBehaviorTreeGraph> GraphClass;
 	FName GraphName;
@@ -273,34 +275,34 @@ protected:
 
 private:
 	/** Create widget for graph editing */
-	TSharedRef<class SGraphEditor> CreateGraphEditorWidget(UEdGraph* InGraph);
+	UE_API TSharedRef<class SGraphEditor> CreateGraphEditorWidget(UEdGraph* InGraph);
 
 	/** Creates all internal widgets for the tabs to point at */
-	void CreateInternalWidgets();
+	UE_API void CreateInternalWidgets();
 
 	/** Add custom menu options */
-	void ExtendMenu();
+	UE_API void ExtendMenu();
 
 	/** Setup common commands */
-	void BindCommonCommands();
+	UE_API void BindCommonCommands();
 
 	/** Setup commands */
-	void BindDebuggerToolbarCommands();
+	UE_API void BindDebuggerToolbarCommands();
 
 	/** Called when the selection changes in the GraphEditor */
-	virtual void OnSelectedNodesChanged(const TSet<class UObject*>& NewSelection) override;
+	UE_API virtual void OnSelectedNodesChanged(const TSet<class UObject*>& NewSelection) override;
 
 	/** prepare range of nodes that can be aborted by this decorator */
-	void GetAbortModePreview(const class UBehaviorTreeGraphNode_CompositeDecorator* Node, struct FAbortDrawHelper& Mode0, struct FAbortDrawHelper& Mode1);
+	UE_API void GetAbortModePreview(const class UBehaviorTreeGraphNode_CompositeDecorator* Node, struct FAbortDrawHelper& Mode0, struct FAbortDrawHelper& Mode1);
 
 	/** prepare range of nodes that can be aborted by this decorator */
-	void GetAbortModePreview(const class UBTDecorator* DecoratorOb, struct FAbortDrawHelper& Mode0, struct FAbortDrawHelper& Mode1);
+	UE_API void GetAbortModePreview(const class UBTDecorator* DecoratorOb, struct FAbortDrawHelper& Mode0, struct FAbortDrawHelper& Mode1);
 
 	/** Refresh the debugger's display */
-	void RefreshDebugger();
+	UE_API void RefreshDebugger();
 
 	/** Push new associated Blackboard data to Blackboard views */
-	void RefreshBlackboardViewsAssociatedObject();
+	UE_API void RefreshBlackboardViewsAssociatedObject();
 
 	TSharedPtr<FDocumentTracker> DocumentManager;
 	TWeakPtr<FDocumentTabFactory> GraphEditorTabFactoryPtr;
@@ -353,9 +355,11 @@ private:
 
 public:
 	/** Modes in mode switcher */
-	static const FName BehaviorTreeMode;
-	static const FName BlackboardMode;
+	static UE_API const FName BehaviorTreeMode;
+	static UE_API const FName BlackboardMode;
 
-	static FText BehaviorTreeModeText;
-	static FText BlackboardModeText;
+	static UE_API FText BehaviorTreeModeText;
+	static UE_API FText BlackboardModeText;
 };
+
+#undef UE_API
diff --git a/Engine/Source/Editor/BehaviorTreeEditor/Public/BehaviorTreeEditorModule.h b/Engine/Source/Editor/BehaviorTreeEditor/Public/BehaviorTreeEditorModule.h
--- a/Engine/Source/Editor/BehaviorTreeEditor/Public/BehaviorTreeEditorModule.h
+++ b/Engine/Source/Editor/BehaviorTreeEditor/Public/BehaviorTreeEditorModule.h
@@ -14,15 +14,15 @@ struct FGraphNodeClassHelper;
 DECLARE_LOG_CATEGORY_EXTERN(LogBehaviorTreeEditor, Log, All);
 
 /** Behavior Tree Editor module */
-class BEHAVIORTREEEDITOR_API FBehaviorTreeEditorModule : public IModuleInterface, public IHasMenuExtensibility, public IHasToolBarExtensibility
+class FBehaviorTreeEditorModule : public IModuleInterface, public IHasMenuExtensibility, public IHasToolBarExtensibility
 {
 public:
 	// IModuleInterface
-	virtual void StartupModule() override;
-	virtual void ShutdownModule() override;
+	BEHAVIORTREEEDITOR_API virtual void StartupModule() override;
+	BEHAVIORTREEEDITOR_API virtual void ShutdownModule() override;
 
 	/** Creates an instance of Behavior Tree editor.  Only virtual so that it can be called across the DLL boundary. */
-	virtual TSharedRef<IBehaviorTreeEditor> CreateBehaviorTreeEditor(const EToolkitMode::Type Mode, const TSharedPtr<IToolkitHost>& InitToolkitHost, UObject* Object);
+	BEHAVIORTREEEDITOR_API virtual TSharedRef<IBehaviorTreeEditor> CreateBehaviorTreeEditor(const EToolkitMode::Type Mode, const TSharedPtr<IToolkitHost>& InitToolkitHost, UObject* Object);
 
 	/** Gets the extensibility managers for outside entities to extend static mesh editor's menus and toolbars */
 	virtual TSharedPtr<FExtensibilityManager> GetMenuExtensibilityManager() override { return MenuExtensibilityManager; }
@@ -31,7 +31,7 @@ public:
 	TSharedPtr<FGraphNodeClassHelper> GetClassCache() const { return ClassCache; }
 
 	/** Behavior Tree app identifier string */
-	static const FName BehaviorTreeEditorAppIdentifier;
+	static BEHAVIORTREEEDITOR_API const FName BehaviorTreeEditorAppIdentifier;
 
 private:
 	TSharedPtr<FExtensibilityManager> MenuExtensibilityManager;
@@ -40,9 +40,3 @@ private:
 	TSharedPtr<FGraphNodeClassHelper> ClassCache;
 };
 
-#if UE_ENABLE_INCLUDE_ORDER_DEPRECATED_IN_5_3
-#include "AIGraphTypes.h"
-#include "CoreMinimal.h"
-#include "Toolkits/IToolkit.h"
-#include "Toolkits/IToolkitHost.h"
-#endif
diff --git a/Engine/Source/Editor/BlueprintEditorLibrary/Public/BlueprintEditorLibrary.h b/Engine/Source/Editor/BlueprintEditorLibrary/Public/BlueprintEditorLibrary.h
--- a/Engine/Source/Editor/BlueprintEditorLibrary/Public/BlueprintEditorLibrary.h
+++ b/Engine/Source/Editor/BlueprintEditorLibrary/Public/BlueprintEditorLibrary.h
@@ -16,6 +16,8 @@
 
 #include "BlueprintEditorLibrary.generated.h"
 
+#define UE_API BLUEPRINTEDITORLIBRARY_API
+
 class FProperty;
 class UBlueprint;
 class UClass;
@@ -41,8 +43,8 @@ enum class EAssetSaveVersionComparisonResults : uint8
 	Older
 };
 
-UCLASS()
-class BLUEPRINTEDITORLIBRARY_API UBlueprintEditorLibrary : public UBlueprintFunctionLibrary
+UCLASS(MinimalAPI)
+class UBlueprintEditorLibrary : public UBlueprintFunctionLibrary
 {
 	GENERATED_UCLASS_BODY()
 public:
@@ -55,7 +57,7 @@ public:
 	* @param NewVarName		The new variable that will be used in the old one's place
 	*/
 	UFUNCTION(BlueprintCallable, Category = "Blueprint Upgrade Tools")
-	static void ReplaceVariableReferences(UBlueprint* Blueprint, const FName OldVarName, const FName NewVarName);
+	static UE_API void ReplaceVariableReferences(UBlueprint* Blueprint, const FName OldVarName, const FName NewVarName);
 
 	/**
 	* Finds the event graph of the given blueprint. Null if it doesn't have one. This will only return
@@ -66,7 +68,7 @@ public:
 	* @return UEdGraph*		Event graph of the blueprint if it has one, null if it doesn't have one
 	*/
 	UFUNCTION(BlueprintCallable, Category = "Blueprint Upgrade Tools")
-	static UEdGraph* FindEventGraph(UBlueprint* Blueprint);
+	static UE_API UEdGraph* FindEventGraph(UBlueprint* Blueprint);
 
 	/**
 	 * Compares the given assets save version to the VersionToCheck. 
@@ -79,7 +81,7 @@ public:
 	 * @param Result			The outcome of the version comparison
 	 */
 	UFUNCTION(BlueprintCallable, Category = "Blueprint Upgrade Tools",  Meta = (ExpandEnumAsExecs = "Result"))
-	static void CompareAssetSaveVersionTo(const UObject* Asset, const FString& VersionToCheck, EAssetSaveVersionComparisonResults& Result);
+	static UE_API void CompareAssetSaveVersionTo(const UObject* Asset, const FString& VersionToCheck, EAssetSaveVersionComparisonResults& Result);
 
 	/**
 	 * Compares the given soft object's save version to the VersionToCheck. This will read the packages file header
@@ -92,7 +94,7 @@ public:
 	 * @param Result			The outcome of the version comparison
 	 */
 	UFUNCTION(BlueprintCallable, Category = "Blueprint Upgrade Tools",  Meta = (ExpandEnumAsExecs = "Result"))
-	static void CompareSoftObjectSaveVersionTo(const TSoftObjectPtr<UObject> ObjectToCheck,  const FString& VersionToCheck, EAssetSaveVersionComparisonResults& Result);
+	static UE_API void CompareSoftObjectSaveVersionTo(const TSoftObjectPtr<UObject> ObjectToCheck,  const FString& VersionToCheck, EAssetSaveVersionComparisonResults& Result);
 
 	/**
 	 * Returns a string representation of the engine version which the given asset was saved with.
@@ -106,7 +108,7 @@ public:
 	 *			For example: "5.6.0-37518009+++UE5+Main"
 	 */
 	UFUNCTION(BlueprintCallable, BlueprintPure, Category = "Blueprint Upgrade Tools")
-	static FString GetSavedByEngineVersion(const UObject* Asset);
+	static UE_API FString GetSavedByEngineVersion(const UObject* Asset);
 
 	/**
 	 * Returns a string which represents the current engine version (FEngineVersion::Current())
@@ -114,7 +116,7 @@ public:
 	 * For example: "5.6.0-37518009+++UE5+Main"
 	 */
 	UFUNCTION(BlueprintCallable, BlueprintPure, Category = "Blueprint Upgrade Tools")
-	static FString GetCurrentEngineVersion();
+	static UE_API FString GetCurrentEngineVersion();
 
 	/**
 	* Finds the graph with the given name on the blueprint. Null if it doesn't have one. 
@@ -125,7 +127,7 @@ public:
 	* @return UEdGraph*		Pointer to the graph with the given name, null if not found
 	*/
 	UFUNCTION(BlueprintCallable, Category = "Blueprint Upgrade Tools")
-	static UEdGraph* FindGraph(UBlueprint* Blueprint, FName GraphName);
+	static UE_API UEdGraph* FindGraph(UBlueprint* Blueprint, FName GraphName);
 
 	/**
 	* Replace any old operator nodes (float + float, vector + float, int + vector, etc)
@@ -135,7 +137,7 @@ public:
 	* @param Blueprint	Blueprint to upgrade
 	*/
 	UFUNCTION(BlueprintCallable, Category = "Blueprint Upgrade Tools")
-	static void UpgradeOperatorNodes(UBlueprint* Blueprint);
+	static UE_API void UpgradeOperatorNodes(UBlueprint* Blueprint);
 
 	/**
 	* Compiles the given blueprint. 
@@ -143,7 +145,7 @@ public:
 	* @param Blueprint	Blueprint to compile
 	*/
 	UFUNCTION(BlueprintCallable, Category = "Blueprint Upgrade Tools")
-	static void CompileBlueprint(UBlueprint* Blueprint);
+	static UE_API void CompileBlueprint(UBlueprint* Blueprint);
 
 	/**
 	* Adds a function to the given blueprint
@@ -154,7 +156,7 @@ public:
 	* @return UEdGraph*
 	*/
 	UFUNCTION(BlueprintCallable, Category = "Blueprint Upgrade Tools")
-	static UEdGraph* AddFunctionGraph(UBlueprint* Blueprint, const FString& FuncName = FString(TEXT("NewFunction")));
+	static UE_API UEdGraph* AddFunctionGraph(UBlueprint* Blueprint, const FString& FuncName = FString(TEXT("NewFunction")));
 
 	/** 
 	* Deletes the function of the given name on this blueprint. Does NOT replace function call sites. 
@@ -163,7 +165,7 @@ public:
 	* @param FuncName		The name of the function to remove
 	*/
 	UFUNCTION(BlueprintCallable, Category = "Blueprint Upgrade Tools")
-	static void RemoveFunctionGraph(UBlueprint* Blueprint, FName FuncName);
+	static UE_API void RemoveFunctionGraph(UBlueprint* Blueprint, FName FuncName);
 
 	/**
 	* Remove any nodes in this blueprint that have no connections made to them.
@@ -171,7 +173,7 @@ public:
 	* @param Blueprint		The blueprint to remove the nodes from
 	*/
 	UFUNCTION(BlueprintCallable, Category = "Blueprint Upgrade Tools")
-	static void RemoveUnusedNodes(UBlueprint* Blueprint);
+	static UE_API void RemoveUnusedNodes(UBlueprint* Blueprint);
 
 	/** 
 	* Removes the given graph from the blueprint if possible 
@@ -180,7 +182,7 @@ public:
 	* @param Graph		The graph to remove
 	*/
 	UFUNCTION(BlueprintCallable, Category = "Blueprint Upgrade Tools")
-	static void RemoveGraph(UBlueprint* Blueprint, UEdGraph* Graph);
+	static UE_API void RemoveGraph(UBlueprint* Blueprint, UEdGraph* Graph);
 
 	/**
 	* Attempts to rename the given graph with a new name
@@ -189,7 +191,7 @@ public:
 	* @param NewNameStr		The new name of the graph
 	*/
 	UFUNCTION(BlueprintCallable, Category = "Blueprint Upgrade Tools")
-	static void RenameGraph(UEdGraph* Graph, const FString& NewNameStr = FString(TEXT("NewGraph")));
+	static UE_API void RenameGraph(UEdGraph* Graph, const FString& NewNameStr = FString(TEXT("NewGraph")));
 
 	/**
 	* Casts the provided Object to a Blueprint - the root asset type of a blueprint asset. Note
@@ -200,7 +202,7 @@ public:
 	* @return UBlueprint*	The blueprint type of the given object, nullptr if the object is not a blueprint.
 	*/
 	UFUNCTION(BlueprintCallable, Category = "Blueprint Upgrade Tools", meta = (Keywords = "cast"))
-	static UBlueprint* GetBlueprintAsset(UObject* Object);
+	static UE_API UBlueprint* GetBlueprintAsset(UObject* Object);
 	
 	/**
 	 * Looks up the UBlueprint that generated the provided class, if any. Provides a 'true' exec pin
@@ -213,15 +215,15 @@ public:
 										is native or otherwise cooked
 	 */
 	UFUNCTION(BlueprintCallable, Category = "Blueprint Upgrade Tools", meta = (ExpandBoolAsExecs = "bDoesClassHaveBlueprint"))
-	static UBlueprint* GetBlueprintForClass(UClass* Class, bool& bDoesClassHaveBlueprint);
+	static UE_API UBlueprint* GetBlueprintForClass(UClass* Class, bool& bDoesClassHaveBlueprint);
 
 	/** Attempt to refresh any open blueprint editors for the given asset */
 	UFUNCTION(BlueprintCallable, Category = "Blueprint Editor")
-	static void RefreshOpenEditorsForBlueprint(const UBlueprint* BP);
+	static UE_API void RefreshOpenEditorsForBlueprint(const UBlueprint* BP);
 
 	/** Refresh any open blueprint editors */
 	UFUNCTION(BlueprintCallable, Category = "Blueprint Editor")
-	static void RefreshAllOpenBlueprintEditors();
+	static UE_API void RefreshAllOpenBlueprintEditors();
 
 	/**
 	* Attempts to reparent the given blueprint to the new chosen parent class. 
@@ -230,7 +232,7 @@ public:
 	* @param NewParentClass		The new parent class to use
 	*/
 	UFUNCTION(BlueprintCallable, Category = "Blueprint Upgrade Tools")
-	static void ReparentBlueprint(UBlueprint* Blueprint, UClass* NewParentClass);
+	static UE_API void ReparentBlueprint(UBlueprint* Blueprint, UClass* NewParentClass);
 
 	/**
 	* Gathers any unused blueprint variables and populates the given array of FPropertys
@@ -240,7 +242,7 @@ public:
 	*
 	* @return					True if variables were checked on this blueprint, false otherwise.
 	*/
-	static bool GatherUnusedVariables(const UBlueprint* Blueprint, TArray<FProperty*>& OutProperties);
+	static UE_API bool GatherUnusedVariables(const UBlueprint* Blueprint, TArray<FProperty*>& OutProperties);
 
 	/**
 	* Deletes any unused blueprint created variables the given blueprint.
@@ -252,7 +254,7 @@ public:
 	* @return					Number of variables removed
 	*/
 	UFUNCTION(BlueprintCallable, Category = "Blueprint Upgrade Tools")
-	static int32 RemoveUnusedVariables(UBlueprint* Blueprint);
+	static UE_API int32 RemoveUnusedVariables(UBlueprint* Blueprint);
 
 	/**
 	 * Gets the class generated when this blueprint is compiled
@@ -261,7 +263,7 @@ public:
 	 * @return UClass*			The generated class
 	 */
 	UFUNCTION(BlueprintCallable, Category = "Editor Scripting", meta = (ScriptMethod))
-	static UClass* GeneratedClass(UBlueprint* BlueprintObj);
+	static UE_API UClass* GeneratedClass(UBlueprint* BlueprintObj);
 
 	/**
 	 * Sets "Expose On Spawn" to true/false on a Blueprint variable
@@ -271,7 +273,7 @@ public:
 	 * @param bExposeOnSpawn	Set to true to expose on spawn
 	 */
 	UFUNCTION(BlueprintCallable, Category = "Editor Scripting", meta = (ScriptMethod))
-	static void SetBlueprintVariableExposeOnSpawn(UBlueprint* Blueprint, const FName& VariableName, bool bExposeOnSpawn);
+	static UE_API void SetBlueprintVariableExposeOnSpawn(UBlueprint* Blueprint, const FName& VariableName, bool bExposeOnSpawn);
 
 	/**
 	 * Sets "Expose To Cinematics" to true/false on a Blueprint variable
@@ -281,7 +283,7 @@ public:
 	 * @param bExposeToCinematics	Set to true to expose to cinematics
 	 */
 	UFUNCTION(BlueprintCallable, Category = "Editor Scripting", meta = (ScriptMethod))
-	static void SetBlueprintVariableExposeToCinematics(UBlueprint* Blueprint, const FName& VariableName, bool bExposeToCinematics);
+	static UE_API void SetBlueprintVariableExposeToCinematics(UBlueprint* Blueprint, const FName& VariableName, bool bExposeToCinematics);
 
 	/**
 	 * Sets "Instance Editable" to true/false on a Blueprint variable
@@ -291,7 +293,7 @@ public:
 	 * @param bInstanceEditable		Toggle InstanceEditable
 	 */
 	UFUNCTION(BlueprintCallable, Category = "Editor Scripting", meta = (ScriptMethod))
-	static void SetBlueprintVariableInstanceEditable(UBlueprint* Blueprint, const FName& VariableName, bool bInstanceEditable);
+	static UE_API void SetBlueprintVariableInstanceEditable(UBlueprint* Blueprint, const FName& VariableName, bool bInstanceEditable);
 
 	/**
 	 * Creates a blueprint based on a specific parent, honoring registered custom blueprint types
@@ -300,7 +302,7 @@ public:
 	 * @param ParentClass			The parent class that the blueprint should be based on
 	 */
 	 UFUNCTION(BlueprintCallable, Category = "Blueprint Editor")
-	 static UBlueprint* CreateBlueprintAssetWithParent(const FString& AssetPath, UClass* ParentClass);
+	 static UE_API UBlueprint* CreateBlueprintAssetWithParent(const FString& AssetPath, UClass* ParentClass);
 
 	/**
 	  * Adds a member variable to the specified blueprint inferring the type from a provided value.
@@ -308,8 +310,8 @@ public:
 	  * @return	true if it succeeds, false if it fails.
 	  */
 	UFUNCTION(BlueprintCallable, CustomThunk, Category = "Blueprint Editor", meta=(CustomStructureParam="DefaultValue"))
-	static bool AddMemberVariableWithValue(UBlueprint* Blueprint, FName MemberName, const int32& DefaultValue);
-	static bool Generic_AddMemberVariableWithValue(UBlueprint* Blueprint, FName MemberName, const uint8* DefaultValuePtr, const FProperty* DefaultValueProp);
+	static UE_API bool AddMemberVariableWithValue(UBlueprint* Blueprint, FName MemberName, const int32& DefaultValue);
+	static UE_API bool Generic_AddMemberVariableWithValue(UBlueprint* Blueprint, FName MemberName, const uint8* DefaultValuePtr, const FProperty* DefaultValueProp);
 	DECLARE_FUNCTION(execAddMemberVariableWithValue);
 	
 	/**
@@ -318,33 +320,35 @@ public:
 	  * @return	true if it succeeds, false if it fails.
 	  */
 	UFUNCTION(BlueprintCallable, Category = "Blueprint Editor")
-	static bool AddMemberVariable(UBlueprint* Blueprint, FName MemberName, const FEdGraphPinType& VariableType);
+	static UE_API bool AddMemberVariable(UBlueprint* Blueprint, FName MemberName, const FEdGraphPinType& VariableType);
 	
 	/** @return a pintype for 'int', 'byte', 'bool', 'real', 'name', 'string' or 'text' - returns 'int' type if invalid type is provided */
 	UFUNCTION(BlueprintCallable, Category = "Blueprint Editor")
-	static FEdGraphPinType GetBasicTypeByName(FName TypeName);
+	static UE_API FEdGraphPinType GetBasicTypeByName(FName TypeName);
 	
 	/** @return a pintype for the provided struct - returns 'int' type if invalid struct is provided */
 	UFUNCTION(BlueprintCallable, Category = "Blueprint Editor")
-	static FEdGraphPinType GetStructType(const UScriptStruct* StructType);
+	static UE_API FEdGraphPinType GetStructType(const UScriptStruct* StructType);
 	
 	/** @return a class reference pintype for the provided class - returns 'int' type if invalid class is provided */
 	UFUNCTION(BlueprintCallable, Category = "Blueprint Editor")
-	static FEdGraphPinType GetClassReferenceType(const UClass* ClassType);
+	static UE_API FEdGraphPinType GetClassReferenceType(const UClass* ClassType);
 	
 	/** @return a object reference pintype for the provided class - returns 'int' type if invalid object type is provided */
 	UFUNCTION(BlueprintCallable, Category = "Blueprint Editor")
-	static FEdGraphPinType GetObjectReferenceType(const UClass* ObjectType);
+	static UE_API FEdGraphPinType GetObjectReferenceType(const UClass* ObjectType);
 	
 	/** @return a array of ContainedType type - returns 'int' type if invalid type is provided */
 	UFUNCTION(BlueprintCallable, Category = "Blueprint Editor")
-	static FEdGraphPinType GetArrayType(const FEdGraphPinType& ContainedType);
+	static UE_API FEdGraphPinType GetArrayType(const FEdGraphPinType& ContainedType);
 	
 	/** @return a set of ContainedType type - returns 'int' type if invalid type is provided */
 	UFUNCTION(BlueprintCallable, Category = "Blueprint Editor")
-	static FEdGraphPinType GetSetType(const FEdGraphPinType& ContainedType);
+	static UE_API FEdGraphPinType GetSetType(const FEdGraphPinType& ContainedType);
 	
 	/** @return a map of KeyType to ValueType type - returns 'int' type if invalid type is provided */
 	UFUNCTION(BlueprintCallable, Category = "Blueprint Editor")
-	static FEdGraphPinType GetMapType(const FEdGraphPinType& KeyType,const FEdGraphPinType& ValueType);
-};
\ No newline at end of file
+	static UE_API FEdGraphPinType GetMapType(const FEdGraphPinType& KeyType,const FEdGraphPinType& ValueType);
+};
+
+#undef UE_API
diff --git a/Engine/Source/Editor/BlueprintGraph/Classes/AssetBlueprintGraphActions.h b/Engine/Source/Editor/BlueprintGraph/Classes/AssetBlueprintGraphActions.h
--- a/Engine/Source/Editor/BlueprintGraph/Classes/AssetBlueprintGraphActions.h
+++ b/Engine/Source/Editor/BlueprintGraph/Classes/AssetBlueprintGraphActions.h
@@ -12,7 +12,7 @@ enum class EK2NewNodeFlags;
  * Essentially this is used to help generate asset nodes/responses for Blueprint graphs and keeping Blueprint code from being dependent on something it shouldn't
  * For example, this can be used to help with dragging/dropping assets on a Blueprint graph
  */
-struct BLUEPRINTGRAPH_API FAssetBlueprintGraphActions
+struct FAssetBlueprintGraphActions
 {
 	virtual ~FAssetBlueprintGraphActions() {}
 
diff --git a/Engine/Source/Editor/BlueprintGraph/Classes/BlueprintAssetNodeSpawner.h b/Engine/Source/Editor/BlueprintGraph/Classes/BlueprintAssetNodeSpawner.h
--- a/Engine/Source/Editor/BlueprintGraph/Classes/BlueprintAssetNodeSpawner.h
+++ b/Engine/Source/Editor/BlueprintGraph/Classes/BlueprintAssetNodeSpawner.h
@@ -13,6 +13,8 @@
 
 #include "BlueprintAssetNodeSpawner.generated.h"
 
+#define UE_API BLUEPRINTGRAPH_API
+
 class UEdGraph;
 class UEdGraphNode;
 class UObject;
@@ -21,8 +23,8 @@ class UObject;
  * Takes care of spawning various asset related nodes (nodes associated with 
  * an asset).
  */
-UCLASS(Transient)
-class BLUEPRINTGRAPH_API UBlueprintAssetNodeSpawner : public UBlueprintNodeSpawner
+UCLASS(MinimalAPI, Transient)
+class UBlueprintAssetNodeSpawner : public UBlueprintNodeSpawner
 {
 	GENERATED_BODY()
 
@@ -38,17 +40,19 @@ public:
 	 * @param  InPostSpawnDelegate	A delegate to perform specialized node setup post-spawn.
 	 * @return A newly allocated instance of this class.
 	 */
-	static UBlueprintAssetNodeSpawner* Create(TSubclassOf<UEdGraphNode> const InNodeClass, const FAssetData& InAssetData, UObject* InOuter = nullptr, FCustomizeNodeDelegate InPostSpawnDelegate = FCustomizeNodeDelegate());
+	static UE_API UBlueprintAssetNodeSpawner* Create(TSubclassOf<UEdGraphNode> const InNodeClass, const FAssetData& InAssetData, UObject* InOuter = nullptr, FCustomizeNodeDelegate InPostSpawnDelegate = FCustomizeNodeDelegate());
 
 	/**
 	 * Retrieves the asset that this assigns to spawned nodes.
 	 *
 	 * @return The asset that this class was initialized with.
 	 */
-	const FAssetData& GetAssetData() const;
+	UE_API const FAssetData& GetAssetData() const;
 
 private:
 	/** The asset to configure new nodes with. */
 	UPROPERTY()
 	FAssetData AssetData;
 };
+
+#undef UE_API
diff --git a/Engine/Source/Editor/BlueprintGraph/Classes/BlueprintBoundEventNodeSpawner.h b/Engine/Source/Editor/BlueprintGraph/Classes/BlueprintBoundEventNodeSpawner.h
--- a/Engine/Source/Editor/BlueprintGraph/Classes/BlueprintBoundEventNodeSpawner.h
+++ b/Engine/Source/Editor/BlueprintGraph/Classes/BlueprintBoundEventNodeSpawner.h
@@ -17,6 +17,8 @@
 
 #include "BlueprintBoundEventNodeSpawner.generated.h"
 
+#define UE_API BLUEPRINTGRAPH_API
+
 class FMulticastDelegateProperty;
 class UBlueprint;
 class UEdGraph;
@@ -31,8 +33,8 @@ class UObject;
  * return the existing one). Evolved from FEdGraphSchemaAction_K2AddEvent and 
  * FEdGraphSchemaAction_K2ViewNode.
  */
-UCLASS(Transient)
-class BLUEPRINTGRAPH_API UBlueprintBoundEventNodeSpawner : public UBlueprintEventNodeSpawner
+UCLASS(MinimalAPI, Transient)
+class UBlueprintBoundEventNodeSpawner : public UBlueprintEventNodeSpawner
 {
 	GENERATED_UCLASS_BODY()
 
@@ -46,27 +48,29 @@ public:
 	 * @param  Outer			Optional outer for the new spawner (if left null, the transient package will be used).
 	 * @return A newly allocated instance of this class.
 	 */
-	static UBlueprintBoundEventNodeSpawner* Create(TSubclassOf<UK2Node_Event> NodeClass, FMulticastDelegateProperty* EventDelegate, UObject* Outer = nullptr);
+	static UE_API UBlueprintBoundEventNodeSpawner* Create(TSubclassOf<UK2Node_Event> NodeClass, FMulticastDelegateProperty* EventDelegate, UObject* Outer = nullptr);
 
 	// UBlueprintNodeSpawner interface
-	virtual FBlueprintNodeSignature GetSpawnerSignature() const override;
-	virtual UEdGraphNode* Invoke(UEdGraph* ParentGraph, FBindingSet const& Bindings, FVector2D const Location) const override;
+	UE_API virtual FBlueprintNodeSignature GetSpawnerSignature() const override;
+	UE_API virtual UEdGraphNode* Invoke(UEdGraph* ParentGraph, FBindingSet const& Bindings, FVector2D const Location) const override;
 	// End UBlueprintNodeSpawner interface
 
 	// UBlueprintEventNodeSpawner interface
-	virtual UK2Node_Event const* FindPreExistingEvent(UBlueprint* Blueprint, FBindingSet const& Bindings) const override;
+	UE_API virtual UK2Node_Event const* FindPreExistingEvent(UBlueprint* Blueprint, FBindingSet const& Bindings) const override;
 	// End UBlueprintEventNodeSpawner interface
 
 	// IBlueprintNodeBinder interface
-	virtual bool IsBindingCompatible(FBindingObject BindingCandidate) const override;
-	virtual bool BindToNode(UEdGraphNode* Node, FBindingObject Binding) const override;
+	UE_API virtual bool IsBindingCompatible(FBindingObject BindingCandidate) const override;
+	UE_API virtual bool BindToNode(UEdGraphNode* Node, FBindingObject Binding) const override;
 	// End IBlueprintNodeBinder interface
 
 	/** @return  */
-	FMulticastDelegateProperty const* GetEventDelegate() const;
+	UE_API FMulticastDelegateProperty const* GetEventDelegate() const;
 
 private:
 	/** */
 	UPROPERTY()
 	TFieldPath<FMulticastDelegateProperty> EventDelegate;
 };
+
+#undef UE_API
diff --git a/Engine/Source/Editor/BlueprintGraph/Classes/BlueprintBoundNodeSpawner.h b/Engine/Source/Editor/BlueprintGraph/Classes/BlueprintBoundNodeSpawner.h
--- a/Engine/Source/Editor/BlueprintGraph/Classes/BlueprintBoundNodeSpawner.h
+++ b/Engine/Source/Editor/BlueprintGraph/Classes/BlueprintBoundNodeSpawner.h
@@ -16,6 +16,8 @@
 
 #include "BlueprintBoundNodeSpawner.generated.h"
 
+#define UE_API BLUEPRINTGRAPH_API
+
 class UBlueprint;
 class UEdGraph;
 class UEdGraphNode;
@@ -25,8 +27,8 @@ class UObject;
  * Takes care of spawning various bound nodes. Acts as the 
  * "action" portion of certain FBlueprintActionMenuItems. 
  */
-UCLASS(Transient)
-class BLUEPRINTGRAPH_API UBlueprintBoundNodeSpawner : public UBlueprintNodeSpawner
+UCLASS(MinimalAPI, Transient)
+class UBlueprintBoundNodeSpawner : public UBlueprintNodeSpawner
 {
 	GENERATED_UCLASS_BODY()
 	DECLARE_DELEGATE_RetVal_OneParam(bool, FCanBindObjectDelegate, UObject const*);
@@ -37,16 +39,16 @@ public:
 	/**
 	 * @return A newly allocated instance of this class.
 	 */
-	static UBlueprintBoundNodeSpawner* Create(TSubclassOf<UEdGraphNode> NodeClass, UObject* Outer = nullptr);
+	static UE_API UBlueprintBoundNodeSpawner* Create(TSubclassOf<UEdGraphNode> NodeClass, UObject* Outer = nullptr);
 
 	// UBlueprintNodeSpawner interface
-	virtual FBlueprintNodeSignature GetSpawnerSignature() const override;
-	virtual UEdGraphNode* Invoke(UEdGraph* ParentGraph, FBindingSet const& Bindings, FVector2D const Location) const override;
+	UE_API virtual FBlueprintNodeSignature GetSpawnerSignature() const override;
+	UE_API virtual UEdGraphNode* Invoke(UEdGraph* ParentGraph, FBindingSet const& Bindings, FVector2D const Location) const override;
 	// End UBlueprintNodeSpawner interface
 
 	// FBlueprintNodeBinder interface
-	virtual bool IsBindingCompatible(FBindingObject BindingCandidate) const override;
-	virtual bool BindToNode(UEdGraphNode* Node, FBindingObject Binding) const override;
+	UE_API virtual bool IsBindingCompatible(FBindingObject BindingCandidate) const override;
+	UE_API virtual bool BindToNode(UEdGraphNode* Node, FBindingObject Binding) const override;
 	// End FBlueprintNodeBinder interface
 
 public:
@@ -66,3 +68,5 @@ public:
 	 */
 	FFindPreExistingNodeDelegate FindPreExistingNodeDelegate;
 };
+
+#undef UE_API
diff --git a/Engine/Source/Editor/BlueprintGraph/Classes/BlueprintComponentNodeSpawner.h b/Engine/Source/Editor/BlueprintGraph/Classes/BlueprintComponentNodeSpawner.h
--- a/Engine/Source/Editor/BlueprintGraph/Classes/BlueprintComponentNodeSpawner.h
+++ b/Engine/Source/Editor/BlueprintGraph/Classes/BlueprintComponentNodeSpawner.h
@@ -17,6 +17,8 @@
 
 #include "BlueprintComponentNodeSpawner.generated.h"
 
+#define UE_API BLUEPRINTGRAPH_API
+
 class UActorComponent;
 class UEdGraph;
 class UEdGraphNode;
@@ -27,8 +29,8 @@ class UObject;
  * portion of certain FBlueprintActionMenuItems. Evolved from 
  * FEdGraphSchemaAction_K2AddComponent.
  */
-UCLASS(Transient)
-class BLUEPRINTGRAPH_API UBlueprintComponentNodeSpawner : public UBlueprintNodeSpawner
+UCLASS(MinimalAPI, Transient)
+class UBlueprintComponentNodeSpawner : public UBlueprintNodeSpawner
 {
 	GENERATED_UCLASS_BODY()
 
@@ -41,18 +43,18 @@ public:
 	 * @param  ComponentClass	The component type you want spawned nodes to spawn.
 	 * @return A newly allocated instance of this class.
 	 */
-	static UBlueprintComponentNodeSpawner* Create(const struct FComponentTypeEntry& Entry);
+	static UE_API UBlueprintComponentNodeSpawner* Create(const struct FComponentTypeEntry& Entry);
 
 	// UBlueprintNodeSpawner interface
-	virtual FBlueprintNodeSignature GetSpawnerSignature() const override;
-	virtual FBlueprintActionUiSpec GetUiSpec(FBlueprintActionContext const& Context, FBindingSet const& Bindings) const override;
-	virtual UEdGraphNode* Invoke(UEdGraph* ParentGraph, FBindingSet const& Bindings, FVector2D const Location) const override;
-	virtual bool IsTemplateNodeFilteredOut(const FBlueprintActionFilter& Filter) const override;
+	UE_API virtual FBlueprintNodeSignature GetSpawnerSignature() const override;
+	UE_API virtual FBlueprintActionUiSpec GetUiSpec(FBlueprintActionContext const& Context, FBindingSet const& Bindings) const override;
+	UE_API virtual UEdGraphNode* Invoke(UEdGraph* ParentGraph, FBindingSet const& Bindings, FVector2D const Location) const override;
+	UE_API virtual bool IsTemplateNodeFilteredOut(const FBlueprintActionFilter& Filter) const override;
 	// End UBlueprintNodeSpawner interface
 	
 	// IBlueprintNodeBinder interface
-	virtual bool IsBindingCompatible(FBindingObject BindingCandidate) const override;
-	virtual bool BindToNode(UEdGraphNode* Node, FBindingObject Binding) const override;
+	UE_API virtual bool IsBindingCompatible(FBindingObject BindingCandidate) const override;
+	UE_API virtual bool BindToNode(UEdGraphNode* Node, FBindingObject Binding) const override;
 	// End IBlueprintNodeBinder interface
 
 	/**
@@ -60,7 +62,7 @@ public:
 	 *
 	 * @return The component type that this class was initialized with.
 	 */
-	TSubclassOf<UActorComponent> GetComponentClass() const;
+	UE_API TSubclassOf<UActorComponent> GetComponentClass() const;
 
 private:
 	/** The component class to configure new nodes with. */
@@ -75,3 +77,5 @@ private:
 	UPROPERTY()
 	FString ComponentAssetName;
 };
+
+#undef UE_API
diff --git a/Engine/Source/Editor/BlueprintGraph/Classes/BlueprintDelegateNodeSpawner.h b/Engine/Source/Editor/BlueprintGraph/Classes/BlueprintDelegateNodeSpawner.h
--- a/Engine/Source/Editor/BlueprintGraph/Classes/BlueprintDelegateNodeSpawner.h
+++ b/Engine/Source/Editor/BlueprintGraph/Classes/BlueprintDelegateNodeSpawner.h
@@ -12,6 +12,8 @@
 
 #include "BlueprintDelegateNodeSpawner.generated.h"
 
+#define UE_API BLUEPRINTGRAPH_API
+
 class FMulticastDelegateProperty;
 class UK2Node_BaseMCDelegate;
 class UObject;
@@ -21,8 +23,8 @@ class UObject;
  * "action" portion for certain FBlueprintActionMenuItems. Evolved from 
  * FEdGraphSchemaAction_K2Delegate, FEdGraphSchemaAction_K2AssignDelegate, etc.
  */
-UCLASS(Transient)
-class BLUEPRINTGRAPH_API UBlueprintDelegateNodeSpawner : public UBlueprintFieldNodeSpawner
+UCLASS(MinimalAPI, Transient)
+class UBlueprintDelegateNodeSpawner : public UBlueprintFieldNodeSpawner
 {
 	GENERATED_UCLASS_BODY()
 
@@ -37,7 +39,7 @@ public:
 	 * @param  Outer		Optional outer for the new spawner (if left null, the transient package will be used).
 	 * @return A newly allocated instance of this class.
 	 */
-	static UBlueprintDelegateNodeSpawner* Create(TSubclassOf<UK2Node_BaseMCDelegate> NodeClass, FMulticastDelegateProperty const* const Property, UObject* Outer = nullptr);
+	static UE_API UBlueprintDelegateNodeSpawner* Create(TSubclassOf<UK2Node_BaseMCDelegate> NodeClass, FMulticastDelegateProperty const* const Property, UObject* Outer = nullptr);
 
 	/**
 	 * Accessor to the delegate property that this spawner wraps (the delegate
@@ -45,5 +47,7 @@ public:
 	 *
 	 * @return The delegate property that this was initialized with.
 	 */
-	FMulticastDelegateProperty const* GetDelegateProperty() const;
+	UE_API FMulticastDelegateProperty const* GetDelegateProperty() const;
 };
+
+#undef UE_API
diff --git a/Engine/Source/Editor/BlueprintGraph/Classes/BlueprintEventNodeSpawner.h b/Engine/Source/Editor/BlueprintGraph/Classes/BlueprintEventNodeSpawner.h
--- a/Engine/Source/Editor/BlueprintGraph/Classes/BlueprintEventNodeSpawner.h
+++ b/Engine/Source/Editor/BlueprintGraph/Classes/BlueprintEventNodeSpawner.h
@@ -17,6 +17,8 @@
 
 #include "BlueprintEventNodeSpawner.generated.h"
 
+#define UE_API BLUEPRINTGRAPH_API
+
 class UBlueprint;
 class UEdGraph;
 class UEdGraphNode;
@@ -31,8 +33,8 @@ class UObject;
  * return the existing one). Evolved from FEdGraphSchemaAction_K2AddEvent and 
  * FEdGraphSchemaAction_K2ViewNode.
  */
-UCLASS(Transient)
-class BLUEPRINTGRAPH_API UBlueprintEventNodeSpawner : public UBlueprintNodeSpawner
+UCLASS(MinimalAPI, Transient)
+class UBlueprintEventNodeSpawner : public UBlueprintNodeSpawner
 {
 	GENERATED_UCLASS_BODY()
 
@@ -46,7 +48,7 @@ public:
 	 * @param  Outer		Optional outer for the new spawner (if left null, the transient package will be used).
 	 * @return A newly allocated instance of this class.
 	 */
-	static UBlueprintEventNodeSpawner* Create(UFunction const* const EventFunc, UObject* Outer = nullptr);
+	static UE_API UBlueprintEventNodeSpawner* Create(UFunction const* const EventFunc, UObject* Outer = nullptr);
 
 	/**
 	 * Creates a new UBlueprintEventNodeSpawner for custom events. The
@@ -58,11 +60,11 @@ public:
 	 * @param  Outer			Optional outer for the new spawner (if left null, the transient package will be used).
 	 * @return A newly allocated instance of this class.
 	 */
-	static UBlueprintEventNodeSpawner* Create(TSubclassOf<UK2Node_Event> NodeClass, FName CustomEventName, UObject* Outer = nullptr);
+	static UE_API UBlueprintEventNodeSpawner* Create(TSubclassOf<UK2Node_Event> NodeClass, FName CustomEventName, UObject* Outer = nullptr);
 
 	// UBlueprintNodeSpawner interface
-	virtual FBlueprintNodeSignature GetSpawnerSignature() const override;
-	virtual UEdGraphNode* Invoke(UEdGraph* ParentGraph, FBindingSet const& Bindings, FVector2D const Location) const override;
+	UE_API virtual FBlueprintNodeSignature GetSpawnerSignature() const override;
+	UE_API virtual UEdGraphNode* Invoke(UEdGraph* ParentGraph, FBindingSet const& Bindings, FVector2D const Location) const override;
 	// End UBlueprintNodeSpawner interface
 
 	/**
@@ -70,7 +72,7 @@ public:
 	 * 
 	 * @return 
 	 */
-	bool IsForCustomEvent() const;
+	UE_API bool IsForCustomEvent() const;
 	
 	/**
 	 * Retrieves the function that this assigns to spawned nodes (defines the
@@ -78,14 +80,14 @@ public:
 	 *
 	 * @return The event function that this class was initialized with.
 	 */
-	UFunction const* GetEventFunction() const;
+	UE_API UFunction const* GetEventFunction() const;
 
 	/**
 	 * 
 	 * 
 	 * @return 
 	 */
-	virtual UK2Node_Event const* FindPreExistingEvent(UBlueprint* Blueprint, FBindingSet const& Bindings) const;
+	UE_API virtual UK2Node_Event const* FindPreExistingEvent(UBlueprint* Blueprint, FBindingSet const& Bindings) const;
 
 protected:
 	/** The function to configure new nodes with. */
@@ -96,3 +98,5 @@ protected:
 	UPROPERTY()
 	FName CustomEventName;	
 };
+
+#undef UE_API
diff --git a/Engine/Source/Editor/BlueprintGraph/Classes/BlueprintFieldNodeSpawner.h b/Engine/Source/Editor/BlueprintGraph/Classes/BlueprintFieldNodeSpawner.h
--- a/Engine/Source/Editor/BlueprintGraph/Classes/BlueprintFieldNodeSpawner.h
+++ b/Engine/Source/Editor/BlueprintGraph/Classes/BlueprintFieldNodeSpawner.h
@@ -19,6 +19,8 @@
 
 #include "BlueprintFieldNodeSpawner.generated.h"
 
+#define UE_API BLUEPRINTGRAPH_API
+
 class FProperty;
 class UClass;
 class UEdGraph;
@@ -32,8 +34,8 @@ class UObject;
  * functions, enums, structs, properties, etc.). Acts as the "action" portion
  * for certain FBlueprintActionMenuItems. 
  */
-UCLASS(Transient)
-class BLUEPRINTGRAPH_API UBlueprintFieldNodeSpawner : public UBlueprintNodeSpawner
+UCLASS(MinimalAPI, Transient)
+class UBlueprintFieldNodeSpawner : public UBlueprintNodeSpawner
 {
 	GENERATED_UCLASS_BODY()
 	DECLARE_DELEGATE_TwoParams(FSetNodeFieldDelegate, UEdGraphNode*, FFieldVariant);
@@ -50,11 +52,11 @@ public:
 	 * @param  OwnerClass	The class that the variable is a member of or the class it is associated with if it is in a sidecar data structure.
 	 * @return A newly allocated instance of this class.
 	 */
-	static UBlueprintFieldNodeSpawner* Create(TSubclassOf<UK2Node> NodeClass, FFieldVariant Field, UObject* Outer = nullptr, UClass const* OwnerClass = nullptr);
+	static UE_API UBlueprintFieldNodeSpawner* Create(TSubclassOf<UK2Node> NodeClass, FFieldVariant Field, UObject* Outer = nullptr, UClass const* OwnerClass = nullptr);
 
 	// UBlueprintNodeSpawner interface
-	virtual FBlueprintNodeSignature GetSpawnerSignature() const override;
-	virtual UEdGraphNode* Invoke(UEdGraph* ParentGraph, FBindingSet const& Bindings, FVector2D const Location) const override;
+	UE_API virtual FBlueprintNodeSignature GetSpawnerSignature() const override;
+	UE_API virtual UEdGraphNode* Invoke(UEdGraph* ParentGraph, FBindingSet const& Bindings, FVector2D const Location) const override;
 	// End UBlueprintNodeSpawner interface
 
 	/** Callback to define how the field should be applied to new nodes */
@@ -66,14 +68,14 @@ public:
 	 *
 	 * @return The field that this class was initialized with.
 	 */
-	FFieldVariant GetField() const;
+	UE_API FFieldVariant GetField() const;
 
 protected:
 
 	/**
 	 * Sets the field for this spawner 
 	 */
-	void SetField(FFieldVariant InField);
+	UE_API void SetField(FFieldVariant InField);
 
 	/** The owning class to configure new nodes with. */
 	UPROPERTY()
@@ -87,3 +89,5 @@ private:
 	UPROPERTY()
 	TFieldPath<FProperty> Property;
 };
+
+#undef UE_API
