[BTDebugger] tentative fix for invalid memory access

diff --git a/Engine/Source/Editor/BehaviorTreeEditor/Private/BehaviorTreeDebugger.cpp b/Engine/Source/Editor/BehaviorTreeEditor/Private/BehaviorTreeDebugger.cpp
--- a/Engine/Source/Editor/BehaviorTreeEditor/Private/BehaviorTreeDebugger.cpp
+++ b/Engine/Source/Editor/BehaviorTreeEditor/Private/BehaviorTreeDebugger.cpp
@@ -143,7 +143,8 @@ void FBehaviorTreeDebugger::Tick(float DeltaTime)
 		return;
 	}
 
-	if (!TreeInstance.IsValid() && !IsPlaySessionPaused())
+	if ((!TreeInstance.IsValid() || !RootNode.IsValid())
+		&& !IsPlaySessionPaused())
 	{
 		// clear state when active tree is lost
 		if (DebuggerInstanceIndex != INDEX_NONE)
@@ -263,20 +264,23 @@ void FBehaviorTreeDebugger::Tick(float DeltaTime)
 		FBehaviorTreeExecutionStep& ExecutionStep = TreeInstance->DebuggerSteps[ActiveStepIndex];
 		const FBehaviorTreeDebuggerInstance& ShowInstance = ExecutionStep.InstanceStack[DebuggerInstanceIndex];
 
-		if (DisplayedExecutionStepId != ExecutionStep.ExecutionStepId)
+		if (UBTNode* InstanceRootNode = ShowInstance.RootNode.Get(); IsValidChecked(InstanceRootNode))
 		{
-			UpdateAssetFlags(ShowInstance, RootNode.Get(), ExecutionStep.ExecutionStepId);
-		}
-		
-		// Collect current runtime descriptions for every node when displaying the most recent step.
-		// Otherwise we keep the descriptions from the displayed debugger step that were updated by UpdateAssetFlags.
-		int32 LatestStepIndex = TreeInstance->DebuggerSteps.Num() - 1;
-		if (ActiveStepIndex == LatestStepIndex)
-		{
-			TArray<FString> RuntimeDescriptions;
-			TreeInstance->StoreDebuggerRuntimeValues(RuntimeDescriptions, ShowInstance.RootNode, IntCastChecked<uint16>(DebuggerInstanceIndex));
+			if (DisplayedExecutionStepId != ExecutionStep.ExecutionStepId)
+			{
+				UpdateAssetFlags(ShowInstance, RootNode.Get(), ExecutionStep.ExecutionStepId);
+			}
+
+			// Collect current runtime descriptions for every node when displaying the most recent step.
+			// Otherwise we keep the descriptions from the displayed debugger step that were updated by UpdateAssetFlags.
+			int32 LatestStepIndex = TreeInstance->DebuggerSteps.Num() - 1;
+			if (ActiveStepIndex == LatestStepIndex)
+			{
+				TArray<FString> RuntimeDescriptions;
+				TreeInstance->StoreDebuggerRuntimeValues(RuntimeDescriptions, InstanceRootNode, IntCastChecked<uint16>(DebuggerInstanceIndex));
 
-			UpdateAssetRuntimeDescription(RuntimeDescriptions, RootNode.Get());
+				UpdateAssetRuntimeDescription(RuntimeDescriptions, RootNode.Get());
+			}
 		}
 	}
 
diff --git a/Engine/Source/Runtime/AIModule/Classes/BehaviorTree/BTNode.h b/Engine/Source/Runtime/AIModule/Classes/BehaviorTree/BTNode.h
--- a/Engine/Source/Runtime/AIModule/Classes/BehaviorTree/BTNode.h
+++ b/Engine/Source/Runtime/AIModule/Classes/BehaviorTree/BTNode.h
@@ -229,7 +229,7 @@ private:
 
 #if USE_BEHAVIORTREE_DEBUGGER
 	/** next node in execution order */
-	UBTNode* NextExecutionNode;
+	TWeakObjectPtr<UBTNode> NextExecutionNode;
 #endif
 
 	/** depth first index (execution order) */
@@ -273,7 +273,7 @@ inline UBTCompositeNode* UBTNode::GetParentNode() const
 #if USE_BEHAVIORTREE_DEBUGGER
 inline UBTNode* UBTNode::GetNextNode() const
 {
-	return NextExecutionNode;
+	return NextExecutionNode.Get();
 }
 #endif
 
diff --git a/Engine/Source/Runtime/AIModule/Classes/BehaviorTree/BehaviorTreeTypes.h b/Engine/Source/Runtime/AIModule/Classes/BehaviorTree/BehaviorTreeTypes.h
--- a/Engine/Source/Runtime/AIModule/Classes/BehaviorTree/BehaviorTreeTypes.h
+++ b/Engine/Source/Runtime/AIModule/Classes/BehaviorTree/BehaviorTreeTypes.h
@@ -246,7 +246,7 @@ struct FBehaviorTreeDebuggerInstance
 	UBehaviorTree* TreeAsset;
 
 	/** root node in template */
-	UBTCompositeNode* RootNode;
+	TWeakObjectPtr<UBTCompositeNode> RootNode;
 
 	/** execution indices of active nodes */
 	TArray<uint16> ActivePath;
