Add "Open Source Code" and "Open Blueprint" context menu actions to Behavior Tree editor nodes for quick navigation to a node's C++ source or Blueprint asset
Chain the base class context menu in composite and task node subclasses so the new entries propagate correctly

#rnx

#Synced-CL 52622581

diff --git a/Engine/Source/Editor/BehaviorTreeEditor/Classes/BehaviorTreeGraphNode.h b/Engine/Source/Editor/BehaviorTreeEditor/Classes/BehaviorTreeGraphNode.h
--- a/Engine/Source/Editor/BehaviorTreeEditor/Classes/BehaviorTreeGraphNode.h
+++ b/Engine/Source/Editor/BehaviorTreeEditor/Classes/BehaviorTreeGraphNode.h
@@ -41,6 +41,7 @@ class UBehaviorTreeGraphNode : public UAIGraphNode
 	UE_API virtual FText GetTooltipText() const override;
 	UE_API virtual bool CanCreateUnderSpecifiedSchema(const UEdGraphSchema* DesiredSchema) const override;
 	UE_API virtual void FindDiffs(class UEdGraphNode* OtherNode, struct FDiffResults& Results) override;
+	UE_API virtual void GetNodeContextMenuActions(class UToolMenu* Menu, class UGraphNodeContextMenuContext* Context) const override;
 	//~ End UEdGraphNode Interface
 
 #if WITH_EDITOR
diff --git a/Engine/Source/Editor/BehaviorTreeEditor/Private/BehaviorTreeGraphNode.cpp b/Engine/Source/Editor/BehaviorTreeEditor/Private/BehaviorTreeGraphNode.cpp
--- a/Engine/Source/Editor/BehaviorTreeEditor/Private/BehaviorTreeGraphNode.cpp
+++ b/Engine/Source/Editor/BehaviorTreeEditor/Private/BehaviorTreeGraphNode.cpp
@@ -15,6 +15,10 @@
 #include "SGraphEditorActionMenuAI.h"
 #include "BehaviorTree/BehaviorTree.h"
 #include "GraphDiffControl.h"
+#include "Editor.h"
+#include "SourceCodeNavigation.h"
+#include "Kismet/KismetSystemLibrary.h"
+#include "Subsystems/AssetEditorSubsystem.h"
 
 #include UE_INLINE_GENERATED_CPP_BY_NAME(BehaviorTreeGraphNode)
 
@@ -395,6 +399,51 @@ void UBehaviorTreeGraphNode::AddContextMenuActionsServices(UToolMenu* Menu, cons
 		FNewToolMenuDelegate::CreateUObject(this, &UBehaviorTreeGraphNode::CreateAddServiceSubMenu, (UEdGraph*)Context->Graph));
 }
 
+void UBehaviorTreeGraphNode::GetNodeContextMenuActions(UToolMenu* Menu, UGraphNodeContextMenuContext* Context) const
+{
+	Super::GetNodeContextMenuActions(Menu, Context);
+
+	const UBTNode* BTNode = Cast<UBTNode>(NodeInstance);
+	if (BTNode == nullptr)
+	{
+		return;
+	}
+
+	const UClass* NodeClass = BTNode->GetClass();
+	if (NodeClass && !NodeClass->ClassGeneratedBy && FSourceCodeNavigation::CanNavigateToClass(NodeClass))
+	{
+		FToolMenuSection& Section = Menu->FindOrAddSection("BehaviorTreeGraphNode_SourceCode");
+		Section.AddMenuEntry(
+			"OpenSourceCode",
+			FText::Format( LOCTEXT("OpenSourceCodeTooltip", "Open Source Code for {0}"), FText::FromString(NodeClass->GetAuthoredName())),
+			LOCTEXT("OpenSourceCodeTooltip", "Open the C++ source file for this node"),
+			FSlateIcon(FAppStyle::GetAppStyleSetName(), "Icons.C++"),
+			FUIAction(FExecuteAction::CreateLambda([NodeClass]()
+			{
+				FSourceCodeNavigation::NavigateToClass(NodeClass);
+			}))
+		);
+	}
+	else if (NodeClass && NodeClass->ClassGeneratedBy)
+	{
+		UObject* BlueprintAsset = NodeClass->ClassGeneratedBy;
+		FToolMenuSection& Section = Menu->FindOrAddSection("BehaviorTreeGraphNode_Blueprint");
+		Section.AddMenuEntry(
+			"OpenBlueprint",
+			FText::Format(LOCTEXT("OpenBlueprintFor", "Open Blueprint for {0}"), FText::FromString(NodeClass->GetAuthoredName())),
+			LOCTEXT("OpenBlueprintTooltip", "Open the Blueprint asset for this node"),
+			FSlateIcon(FAppStyle::GetAppStyleSetName(), "Icons.Blueprint"),
+			FUIAction(FExecuteAction::CreateLambda([BlueprintAsset]()
+			{
+				if (BlueprintAsset)
+				{
+					GEditor->GetEditorSubsystem<UAssetEditorSubsystem>()->OpenEditorForAsset(BlueprintAsset);
+				}
+			}))
+		);
+	}
+}
+
 void UBehaviorTreeGraphNode::ClearDebuggerState()
 {
 	bHasBreakpoint = false;
diff --git a/Engine/Source/Editor/BehaviorTreeEditor/Private/BehaviorTreeGraphNode_Composite.cpp b/Engine/Source/Editor/BehaviorTreeEditor/Private/BehaviorTreeGraphNode_Composite.cpp
--- a/Engine/Source/Editor/BehaviorTreeEditor/Private/BehaviorTreeGraphNode_Composite.cpp
+++ b/Engine/Source/Editor/BehaviorTreeEditor/Private/BehaviorTreeGraphNode_Composite.cpp
@@ -77,6 +77,8 @@ void UBehaviorTreeGraphNode_Composite::PostPasteNode()
 
 void UBehaviorTreeGraphNode_Composite::GetNodeContextMenuActions(UToolMenu* Menu, UGraphNodeContextMenuContext* Context) const
 {
+	Super::GetNodeContextMenuActions(Menu, Context);
+
 	AddContextMenuActionsDecorators(Menu, "BehaviorTreeGraphNode", Context);
 
 	if (GetOwnerBehaviorTreeGraph()->DoesSupportServices())
diff --git a/Engine/Source/Editor/BehaviorTreeEditor/Private/BehaviorTreeGraphNode_Task.cpp b/Engine/Source/Editor/BehaviorTreeEditor/Private/BehaviorTreeGraphNode_Task.cpp
--- a/Engine/Source/Editor/BehaviorTreeEditor/Private/BehaviorTreeGraphNode_Task.cpp
+++ b/Engine/Source/Editor/BehaviorTreeEditor/Private/BehaviorTreeGraphNode_Task.cpp
@@ -45,6 +45,8 @@ FText UBehaviorTreeGraphNode_Task::GetNodeTitle(ENodeTitleType::Type TitleType)
 
 void UBehaviorTreeGraphNode_Task::GetNodeContextMenuActions(UToolMenu* Menu, UGraphNodeContextMenuContext* Context) const
 {
+	Super::GetNodeContextMenuActions(Menu, Context);
+
 	AddContextMenuActionsDecorators(Menu, "BehaviorTreeGraphNode", Context);
 
 	if (GetOwnerBehaviorTreeGraph()->DoesSupportServices())
