Fix loctext for BehaviorTreeGraphNode
#rnx
#Synced-CL 52721423
#Synced-CL 52721423

diff --git a/Engine/Source/Editor/BehaviorTreeEditor/Private/BehaviorTreeGraphNode.cpp b/Engine/Source/Editor/BehaviorTreeEditor/Private/BehaviorTreeGraphNode.cpp
--- a/Engine/Source/Editor/BehaviorTreeEditor/Private/BehaviorTreeGraphNode.cpp
+++ b/Engine/Source/Editor/BehaviorTreeEditor/Private/BehaviorTreeGraphNode.cpp
@@ -415,7 +415,7 @@ void UBehaviorTreeGraphNode::GetNodeContextMenuActions(UToolMenu* Menu, UGraphNo
 		FToolMenuSection& Section = Menu->FindOrAddSection("BehaviorTreeGraphNode_SourceCode");
 		Section.AddMenuEntry(
 			"OpenSourceCode",
-			FText::Format( LOCTEXT("OpenSourceCodeTooltip", "Open Source Code for {0}"), FText::FromString(NodeClass->GetAuthoredName())),
+			FText::Format( LOCTEXT("OpenSourceCodeFor", "Open Source Code for {0}"), FText::FromString(NodeClass->GetAuthoredName())),
 			LOCTEXT("OpenSourceCodeTooltip", "Open the C++ source file for this node"),
 			FSlateIcon(FAppStyle::GetAppStyleSetName(), "Icons.C++"),
 			FUIAction(FExecuteAction::CreateLambda([NodeClass]()
